Naming Tokens: Create NamingTokens plugin for defining and evaluating tokens in a string, such as in a filepath.

diff --git a/Engine/Plugins/Animation/LiveLinkHub/LiveLinkHub.uplugin b/Engine/Plugins/Animation/LiveLinkHub/LiveLinkHub.uplugin
--- a/Engine/Plugins/Animation/LiveLinkHub/LiveLinkHub.uplugin
+++ b/Engine/Plugins/Animation/LiveLinkHub/LiveLinkHub.uplugin
@@ -20,6 +20,10 @@
 		{
 			"Name": "LiveLink",
 			"Enabled": true
+		},
+		{
+			"Name": "NamingTokens",
+			"Enabled": true
 		}
 	],
 	"Modules": [
diff --git a/Engine/Plugins/Animation/LiveLinkHub/Source/LiveLinkHub/LiveLinkHub.Build.cs b/Engine/Plugins/Animation/LiveLinkHub/Source/LiveLinkHub/LiveLinkHub.Build.cs
--- a/Engine/Plugins/Animation/LiveLinkHub/Source/LiveLinkHub/LiveLinkHub.Build.cs
+++ b/Engine/Plugins/Animation/LiveLinkHub/Source/LiveLinkHub/LiveLinkHub.Build.cs
@@ -46,6 +46,7 @@ public class LiveLinkHub : ModuleRules
 				"ToolWidgets",
 				"UnrealEd",
 				"OutputLog",
+				"NamingTokens",
 			});
 
 		if (Target.Configuration != UnrealTargetConfiguration.Shipping)
diff --git a/Engine/Plugins/Animation/LiveLinkHub/Source/LiveLinkHub/Private/Config/LiveLinkHubFileUtilities.cpp b/Engine/Plugins/Animation/LiveLinkHub/Source/LiveLinkHub/Private/Config/LiveLinkHubFileUtilities.cpp
--- a/Engine/Plugins/Animation/LiveLinkHub/Source/LiveLinkHub/Private/Config/LiveLinkHubFileUtilities.cpp
+++ b/Engine/Plugins/Animation/LiveLinkHub/Source/LiveLinkHub/Private/Config/LiveLinkHubFileUtilities.cpp
@@ -86,59 +86,3 @@ ULiveLinkHubSessionData* UE::LiveLinkHub::FileUtilities::Private::FromJson(const
 
 	return nullptr;
 }
-
-void UE::LiveLinkHub::FileUtilities::Private::ParseFilenameTemplate(const FString& InFilenameTemplate,
-	FFilenameTemplateData& OutTemplateData)
-{
-	FString FormattedString = InFilenameTemplate;
-	
-	// Replace tokens.
-	{
-		// Get current datetime.
-		const FDateTime CurrentDate = FDateTime::Now();
-		const FString Year2DigitValue = FString::Printf(TEXT("%02d"), CurrentDate.GetYear() % 100);
-		const FString Year4DigitValue = FString::Printf(TEXT("%02d"), CurrentDate.GetYear());
-		const FString MonthValue = FString::Printf(TEXT("%02d"), CurrentDate.GetMonth());
-		const FString DayValue = FString::Printf(TEXT("%02d"), CurrentDate.GetDay());
-		const FString HourValue = FString::Printf(TEXT("%02d"), CurrentDate.GetHour());
-		const FString MinuteValue = FString::Printf(TEXT("%02d"), CurrentDate.GetMinute());
-
-		const FLiveLinkHubAutomaticTokens& AutomaticTokens = FLiveLinkHubAutomaticTokens::GetStaticTokens();
-
-		using namespace UE::LiveLinkHub::Tokens::Private;
-		
-		FormattedString.ReplaceInline(*CreateToken(AutomaticTokens.Year4Digit), *Year4DigitValue, ESearchCase::IgnoreCase);
-		FormattedString.ReplaceInline(*CreateToken(AutomaticTokens.Year2Digit), *Year2DigitValue, ESearchCase::IgnoreCase);
-		FormattedString.ReplaceInline(*CreateToken(AutomaticTokens.Month), *MonthValue, ESearchCase::CaseSensitive);
-		FormattedString.ReplaceInline(*CreateToken(AutomaticTokens.Day), *DayValue, ESearchCase::IgnoreCase);
-		FormattedString.ReplaceInline(*CreateToken(AutomaticTokens.Hour), *HourValue, ESearchCase::IgnoreCase);
-		FormattedString.ReplaceInline(*CreateToken(AutomaticTokens.Minute), *MinuteValue, ESearchCase::CaseSensitive);
-
-		// Get session information.
-		FString SessionNameValue;
-		if (TSharedPtr<ILiveLinkHubSessionManager> SessionManager = FModuleManager::Get().GetModuleChecked<FLiveLinkHubModule>("LiveLinkHub").GetLiveLinkHub()->GetSessionManager())
-		{
-			SessionNameValue = FPaths::GetBaseFilename(SessionManager->GetLastConfigPath());
-		}
-		FormattedString.ReplaceInline(*CreateToken(AutomaticTokens.SessionName), *SessionNameValue);
-	}
-	
-	OutTemplateData.FullPath = FormattedString;
-
-	// Split folder path and file name.
-	int32 LastSlashIndex;
-	if (FormattedString.FindLastChar('/', LastSlashIndex))
-	{
-		OutTemplateData.FolderPath = FormattedString.Left(LastSlashIndex);
-		if (OutTemplateData.FolderPath.StartsWith("/"))
-		{
-			OutTemplateData.FolderPath.RemoveFromStart("/");
-		}
-		OutTemplateData.FileName = FormattedString.Mid(LastSlashIndex + 1);
-	}
-	else
-	{
-		OutTemplateData.FolderPath = TEXT("");
-		OutTemplateData.FileName = FormattedString;
-	}
-}
diff --git a/Engine/Plugins/Animation/LiveLinkHub/Source/LiveLinkHub/Private/Config/LiveLinkHubFileUtilities.h b/Engine/Plugins/Animation/LiveLinkHub/Source/LiveLinkHub/Private/Config/LiveLinkHubFileUtilities.h
--- a/Engine/Plugins/Animation/LiveLinkHub/Source/LiveLinkHub/Private/Config/LiveLinkHubFileUtilities.h
+++ b/Engine/Plugins/Animation/LiveLinkHub/Source/LiveLinkHub/Private/Config/LiveLinkHubFileUtilities.h
@@ -44,7 +44,4 @@ namespace UE::LiveLinkHub::FileUtilities::Private
 		/** The evaluated filename. */
 		FString FileName;
 	};
-
-	/** Parse a filename template for included template tokens and output the computed result. */
-	void ParseFilenameTemplate(const FString& InFilenameTemplate, FFilenameTemplateData& OutTemplateData);
 }
diff --git a/Engine/Plugins/Animation/LiveLinkHub/Source/LiveLinkHub/Private/Config/LiveLinkHubTemplateTokens.cpp b/Engine/Plugins/Animation/LiveLinkHub/Source/LiveLinkHub/Private/Config/LiveLinkHubTemplateTokens.cpp
new file mode 100644
--- /dev/null
+++ b/Engine/Plugins/Animation/LiveLinkHub/Source/LiveLinkHub/Private/Config/LiveLinkHubTemplateTokens.cpp
@@ -0,0 +1,40 @@
+// Copyright Epic Games, Inc. All Rights Reserved.
+
+#include "LiveLinkHubTemplateTokens.h"
+
+#include "Session/LiveLinkHubSessionManager.h"
+
+#define LOCTEXT_NAMESPACE "LiveLinkHubNamingTokens"
+
+ULiveLinkHubNamingTokens::ULiveLinkHubNamingTokens()
+{
+	Namespace = TEXT("llh");
+	NamingTokenContextClass = ULiveLinkHubNamingTokenContext::StaticClass();
+}
+
+void ULiveLinkHubNamingTokens::OnCreateDefaultTokens(TArray<FNamingTokenData>& Tokens)
+{
+	Super::OnCreateDefaultTokens(Tokens);
+
+	Tokens.Add({
+		LOCTEXT("SessionName", "Session Name"),
+		TEXT("session"),
+		FTokenProcessorDelegateNative::CreateLambda([](const TObjectPtr<UNamingTokenContext>& Context)
+		{
+			return FText::FromString(FString::Printf(TEXT("%s"), *CastChecked<ULiveLinkHubNamingTokenContext>(Context)->SessionName));
+		})
+	});
+}
+
+UNamingTokenContext* ULiveLinkHubNamingTokens::CreateContext_Implementation() const
+{
+	ULiveLinkHubNamingTokenContext* Context = CastChecked<ULiveLinkHubNamingTokenContext>(Super::CreateContext_Implementation());
+	if (const TSharedPtr<ILiveLinkHubSessionManager> SessionManager = FModuleManager::Get().GetModuleChecked<FLiveLinkHubModule>("LiveLinkHub").GetLiveLinkHub()->GetSessionManager())
+	{
+		Context->SessionName = FPaths::GetBaseFilename(SessionManager->GetLastConfigPath());
+	}
+	
+	return Context;
+}
+
+#undef LOCTEXT_NAMESPACE
\ No newline at end of file
diff --git a/Engine/Plugins/Animation/LiveLinkHub/Source/LiveLinkHub/Private/Config/LiveLinkHubTemplateTokens.h b/Engine/Plugins/Animation/LiveLinkHub/Source/LiveLinkHub/Private/Config/LiveLinkHubTemplateTokens.h
--- a/Engine/Plugins/Animation/LiveLinkHub/Source/LiveLinkHub/Private/Config/LiveLinkHubTemplateTokens.h
+++ b/Engine/Plugins/Animation/LiveLinkHub/Source/LiveLinkHub/Private/Config/LiveLinkHubTemplateTokens.h
@@ -2,61 +2,33 @@
 
 #pragma once
 
-#include "CoreMinimal.h"
-#include "UObject/ObjectMacros.h"
+#include "NamingTokenContext.h"
+#include "NamingTokens.h"
 
 #include "LiveLinkHubTemplateTokens.generated.h"
 
-/**
- * Predefined tokens used in template parsing. Set up as a USTRUCT so the properties and tooltips
- * can be easily iterated and displayed to the user.
- */
-USTRUCT()
-struct FLiveLinkHubAutomaticTokens
+UCLASS()
+class ULiveLinkHubNamingTokenContext : public UNamingTokenContext
 {
 	GENERATED_BODY()
 
-	/** Two-digit year. */
-	UPROPERTY()
-	FString Year2Digit = TEXT("yy");
-	
-	/** Four-digit year. */
-	UPROPERTY()
-	FString Year4Digit = TEXT("yyyy");
-	
-	/** Two-digit month. */
-	UPROPERTY()
-	FString Month = TEXT("MM");
-	
-	/** Two-digit day. */
-	UPROPERTY()
-	FString Day = TEXT("dd");
-	
-	/** Two-digit hour. */
-	UPROPERTY()
-	FString Hour = TEXT("hh");
-	
-	/** Two-digit minute. */
-	UPROPERTY()
-	FString Minute = TEXT("mm");
-
-	/** The current session name. */
-	UPROPERTY()
-	FString SessionName = TEXT("sessionName");
-
-	/** Singleton like access to this ustruct. */
-	static FLiveLinkHubAutomaticTokens& GetStaticTokens()
-	{
-		static FLiveLinkHubAutomaticTokens AutomaticTokens;
-		return AutomaticTokens;
-	}
+public:
+	/** The name of the current Live Link Hub session. */
+	UPROPERTY(VisibleAnywhere, BlueprintReadOnly, Category = "NamingTokens")
+	FString SessionName;
 };
 
-namespace UE::LiveLinkHub::Tokens::Private
+UCLASS()
+class ULiveLinkHubNamingTokens : public UNamingTokens
 {
-	/** Return the token formatted with curly brackets. */
-	inline FString CreateToken(const FString& InString)
-	{
-		return FString::Printf(TEXT("{%s}"), *InString);
-	}
-}
\ No newline at end of file
+	GENERATED_BODY()
+
+public:
+	ULiveLinkHubNamingTokens();
+
+protected:
+	virtual void OnCreateDefaultTokens(TArray<FNamingTokenData>& Tokens) override;
+	
+	virtual UNamingTokenContext* CreateContext_Implementation() const override;
+	
+};
diff --git a/Engine/Plugins/Animation/LiveLinkHub/Source/LiveLinkHub/Private/Recording/Implementations/LiveLinkUAssetRecorder.cpp b/Engine/Plugins/Animation/LiveLinkHub/Source/LiveLinkHub/Private/Recording/Implementations/LiveLinkUAssetRecorder.cpp
--- a/Engine/Plugins/Animation/LiveLinkHub/Source/LiveLinkHub/Private/Recording/Implementations/LiveLinkUAssetRecorder.cpp
+++ b/Engine/Plugins/Animation/LiveLinkHub/Source/LiveLinkHub/Private/Recording/Implementations/LiveLinkUAssetRecorder.cpp
@@ -20,6 +20,7 @@
 #include "LiveLinkTypes.h"
 #include "Misc/FileHelper.h"
 #include "Misc/PackageName.h"
+#include "NamingTokensEditorSubsystem.h"
 #include "Recording/Implementations/LiveLinkUAssetRecording.h"
 #include "Recording/LiveLinkHubPlaybackController.h"
 #include "Recording/LiveLinkRecording.h"
@@ -151,14 +152,13 @@ bool FLiveLinkUAssetRecorder::OpenSaveDialog(const FString& InDefaultPath, const
 
 bool FLiveLinkUAssetRecorder::GetSavePresetPackageName(FString& OutName)
 {
-	using namespace UE::LiveLinkHub::FileUtilities::Private;
 	const FDateTime Today = FDateTime::Now();
 
 	TMap<FString, FStringFormatArg> FormatArgs;
 	FormatArgs.Add(TEXT("date"), Today.ToString());
-	
-	FFilenameTemplateData TemplateData;
-	ParseFilenameTemplate(GetDefault<ULiveLinkHubSettings>()->FilenameTemplate, TemplateData);
+
+	const TObjectPtr<ULiveLinkHubNamingTokens> Tokens = GetDefault<ULiveLinkHubSettings>()->NamingTokens;
+	FNamingTokenResultData TemplateData = GEditor->GetEditorSubsystem<UNamingTokensEditorSubsystem>()->EvaluateTokenString(GetDefault<ULiveLinkHubSettings>()->FilenameTemplate, Tokens);
 
 	const FString DefaultName = TemplateData.FileName;
 	const FString DefaultFolder = TemplateData.FolderPath;
diff --git a/Engine/Plugins/Animation/LiveLinkHub/Source/LiveLinkHub/Private/Settings/LiveLinkHubSettings.cpp b/Engine/Plugins/Animation/LiveLinkHub/Source/LiveLinkHub/Private/Settings/LiveLinkHubSettings.cpp
--- a/Engine/Plugins/Animation/LiveLinkHub/Source/LiveLinkHub/Private/Settings/LiveLinkHubSettings.cpp
+++ b/Engine/Plugins/Animation/LiveLinkHub/Source/LiveLinkHub/Private/Settings/LiveLinkHubSettings.cpp
@@ -4,6 +4,10 @@
 
 #include "Config/LiveLinkHubFileUtilities.h"
 #include "Config/LiveLinkHubTemplateTokens.h"
+#include "NamingTokensEditorSubsystem.h"
+
+#include "Editor.h"
+#include "UObject/Package.h"
 
 void ULiveLinkHubSettings::PostEditChangeProperty(FPropertyChangedEvent& PropertyChangedEvent)
 {
@@ -12,12 +16,47 @@ void ULiveLinkHubSettings::PostEditChangeProperty(FPropertyChangedEvent& Propert
 	{
 		CalculateExampleOutput();
 	}
+	else if (PropertyChangedEvent.GetPropertyName() == GET_MEMBER_NAME_CHECKED(ULiveLinkHubSettings, NamingTokensClass))
+	{
+		CreateNamingTokens();
+		CalculateExampleOutput();
+	}
+}
+
+void ULiveLinkHubSettings::CreateNamingTokens()
+{
+	if (NamingTokensClass)
+	{
+		NamingTokens = NewObject<ULiveLinkHubNamingTokens>(GetTransientPackage(), NamingTokensClass.Get());
+		NamingTokens->CreateDefaultTokens();
+	}
+	else
+	{
+		NamingTokens = nullptr;
+	}
+}
+
+ULiveLinkHubSettings::ULiveLinkHubSettings()
+{
+	NamingTokensClass = ULiveLinkHubNamingTokens::StaticClass();
+}
+
+void ULiveLinkHubSettings::PostInitProperties()
+{
+	Super::PostInitProperties();
+	CreateNamingTokens();
 }
 
 void ULiveLinkHubSettings::CalculateExampleOutput()
 {
-	using namespace UE::LiveLinkHub::FileUtilities::Private;
-	FFilenameTemplateData TemplateData;
-	ParseFilenameTemplate(FilenameTemplate, TemplateData);
-	FilenameOutput = TemplateData.FullPath;
+	if (const TObjectPtr<ULiveLinkHubNamingTokens> Tokens = GetDefault<ULiveLinkHubSettings>()->NamingTokens)
+	{
+		const FNamingTokenResultData TemplateData = GEditor->GetEditorSubsystem<UNamingTokensEditorSubsystem>()->EvaluateTokenString(
+			GetDefault<ULiveLinkHubSettings>()->FilenameTemplate,Tokens);
+		FilenameOutput = TemplateData.EvaluatedText.ToString();
+	}
+	else
+	{
+		FilenameOutput = GetDefault<ULiveLinkHubSettings>()->FilenameTemplate;
+	}
 }
diff --git a/Engine/Plugins/Animation/LiveLinkHub/Source/LiveLinkHub/Private/Settings/LiveLinkHubSettings.h b/Engine/Plugins/Animation/LiveLinkHub/Source/LiveLinkHub/Private/Settings/LiveLinkHubSettings.h
--- a/Engine/Plugins/Animation/LiveLinkHub/Source/LiveLinkHub/Private/Settings/LiveLinkHubSettings.h
+++ b/Engine/Plugins/Animation/LiveLinkHub/Source/LiveLinkHub/Private/Settings/LiveLinkHubSettings.h
@@ -2,9 +2,11 @@
 
 #pragma once
 
+#include "Config/LiveLinkHubTemplateTokens.h"
 #include "UObject/Object.h"
 #include "LiveLinkHubMessages.h"
 #include "Misc/FrameRate.h"
+#include "Templates/SubclassOf.h"
 
 #include "LiveLinkHubSettings.generated.h"
 
@@ -17,6 +19,10 @@ class LIVELINKHUB_API ULiveLinkHubSettings : public UObject
 	GENERATED_BODY()
 
 public:
+	ULiveLinkHubSettings();
+	
+	virtual void PostInitProperties() override;
+	
 	/** Parse templates and set example output fields. */
 	void CalculateExampleOutput();
 	
@@ -24,6 +30,9 @@ protected:
 	//~ Begin UObject interface
 	virtual void PostEditChangeProperty(FPropertyChangedEvent& PropertyChangedEvent) override;
 	//~ End UObject interface
+
+	/** Create or destroy naming tokens based on the naming tokens class. */
+	void CreateNamingTokens();
 	
 public:
 	/** If enabled, discovered clients will be automatically added to the current session. */
@@ -70,6 +79,17 @@ public:
 	UPROPERTY(VisibleAnywhere, Category="Templates", DisplayName="Output")
 	FString FilenameOutput;
 
+	/** The class to use for the naming tokens. */
+	UPROPERTY(config, EditAnywhere, NoClear, Category="Templates")
+	TSubclassOf<ULiveLinkHubNamingTokens> NamingTokensClass;
+
+	/**
+	 * Naming tokens for Live Link, instantiated each load based on the naming tokens class.
+	 * This isn't serialized to the config file, and exists here for singleton-like access.
+	 */
+	UPROPERTY(Instanced)
+	TObjectPtr<ULiveLinkHubNamingTokens> NamingTokens;
+	
 	/** Placeholder for a list of the automatic tokens, set from the customization. */
 	UPROPERTY(VisibleAnywhere, Category="Templates")
 	FText AutomaticTokens;
diff --git a/Engine/Plugins/Animation/LiveLinkHub/Source/LiveLinkHub/Private/Settings/LiveLinkHubSettingsCustomization.cpp b/Engine/Plugins/Animation/LiveLinkHub/Source/LiveLinkHub/Private/Settings/LiveLinkHubSettingsCustomization.cpp
--- a/Engine/Plugins/Animation/LiveLinkHub/Source/LiveLinkHub/Private/Settings/LiveLinkHubSettingsCustomization.cpp
+++ b/Engine/Plugins/Animation/LiveLinkHub/Source/LiveLinkHub/Private/Settings/LiveLinkHubSettingsCustomization.cpp
@@ -6,7 +6,9 @@
 #include "DetailCategoryBuilder.h"
 #include "DetailLayoutBuilder.h"
 #include "DetailWidgetRow.h"
+#include "Editor.h"
 #include "LiveLinkHubSettings.h"
+#include "NamingTokensEditorSubsystem.h"
 #include "Widgets/Text/STextBlock.h"
 
 #define LOCTEXT_NAMESPACE "LiveLinkHubSettingsCustomization"
@@ -27,25 +29,6 @@ void FLiveLinkHubSettingsCustomization::CustomizeDetails(IDetailLayoutBuilder& D
 	IDetailPropertyRow* AutomaticTokensRow = DetailBuilder.EditDefaultProperty(AutomaticTokensHandle);
 	check(AutomaticTokensRow);
 	
-	const FLiveLinkHubAutomaticTokens& AutomaticTokens = FLiveLinkHubAutomaticTokens::GetStaticTokens();
-	FString FormattedTokensString;
-	
-	for (TFieldIterator<FProperty> It(FLiveLinkHubAutomaticTokens::StaticStruct()); It; ++It)
-	{
-		FProperty* Property = *It;
-		if (const FStrProperty* StrProperty = CastField<FStrProperty>(Property))
-		{
-			const FString PropertyValue = StrProperty->GetPropertyValue_InContainer(&AutomaticTokens);
-			const FString Token = UE::LiveLinkHub::Tokens::Private::CreateToken(PropertyValue);
-
-			const FString Tooltip = Property->GetToolTipText().ToString();
-			if (!Tooltip.IsEmpty())
-			{
-				FormattedTokensString += FString::Printf(TEXT("%s - %s\n"), *Token, *Tooltip);
-			}
-		}
-	}
-
 	AutomaticTokensRow->CustomWidget()
 	.WholeRowContent()
 	[
@@ -62,10 +45,25 @@ void FLiveLinkHubSettingsCustomization::CustomizeDetails(IDetailLayoutBuilder& D
 		.AutoHeight()
 		[
 			SNew(STextBlock)
-			.Text(FText::FromString(FormattedTokensString))
+			.Text_Raw(this, &FLiveLinkHubSettingsCustomization::GetDisplayTokenText)
 			.Font(IDetailLayoutBuilder::GetDetailFont())
 		]
 	];
 }
 
+FText FLiveLinkHubSettingsCustomization::GetDisplayTokenText() const
+{
+	FString FormattedTokensString = TEXT("None");
+	if (GEditor)
+	{
+		FNamingTokenFilterArgs Args;
+		if (const TObjectPtr<ULiveLinkHubNamingTokens> Tokens = GetDefault<ULiveLinkHubSettings>()->NamingTokens)
+		{
+			Args.NamespacesToInclude.Add(Tokens->GetNamespace());
+		}
+		FormattedTokensString = GEditor->GetEditorSubsystem<UNamingTokensEditorSubsystem>()->GetFormattedTokensStringForDisplay(MoveTemp(Args));
+	}
+	return FText::FromString(FormattedTokensString);
+}
+
 #undef LOCTEXT_NAMESPACE
\ No newline at end of file
diff --git a/Engine/Plugins/Animation/LiveLinkHub/Source/LiveLinkHub/Private/Settings/LiveLinkHubSettingsCustomization.h b/Engine/Plugins/Animation/LiveLinkHub/Source/LiveLinkHub/Private/Settings/LiveLinkHubSettingsCustomization.h
--- a/Engine/Plugins/Animation/LiveLinkHub/Source/LiveLinkHub/Private/Settings/LiveLinkHubSettingsCustomization.h
+++ b/Engine/Plugins/Animation/LiveLinkHub/Source/LiveLinkHub/Private/Settings/LiveLinkHubSettingsCustomization.h
@@ -13,4 +13,7 @@ class FLiveLinkHubSettingsCustomization : public IDetailCustomization
 public:
 	static TSharedRef<IDetailCustomization> MakeInstance();
 	virtual void CustomizeDetails(IDetailLayoutBuilder& DetailBuilder) override;
+	
+private:
+	FText GetDisplayTokenText() const;
 };
\ No newline at end of file
diff --git a/Engine/Plugins/Developer/NamingTokens/NamingTokens.uplugin b/Engine/Plugins/Developer/NamingTokens/NamingTokens.uplugin
new file mode 100644
--- /dev/null
+++ b/Engine/Plugins/Developer/NamingTokens/NamingTokens.uplugin
@@ -0,0 +1,23 @@
+﻿{
+	"FileVersion": 3,
+	"Version": 1,
+	"VersionName": "1.0",
+	"FriendlyName": "NamingTokens",
+	"Description": "Define tokens which can be recognized during string evaluation.",
+	"Category": "Editor",
+	"CreatedBy" : "Epic Games, Inc.",
+	"CreatedByURL" : "http://epicgames.com",
+	"DocsURL": "",
+	"MarketplaceURL": "",
+	"CanContainContent": true,
+	"IsBetaVersion": true,
+	"IsExperimentalVersion": false,
+	"Installed": false,
+	"Modules": [
+		{
+			"Name": "NamingTokens",
+			"Type": "UncookedOnly",
+			"LoadingPhase": "Default"
+		}
+	]
+}
\ No newline at end of file
diff --git a/Engine/Plugins/Developer/NamingTokens/Source/NamingTokens/NamingTokens.Build.cs b/Engine/Plugins/Developer/NamingTokens/Source/NamingTokens/NamingTokens.Build.cs
new file mode 100644
--- /dev/null
+++ b/Engine/Plugins/Developer/NamingTokens/Source/NamingTokens/NamingTokens.Build.cs
@@ -0,0 +1,32 @@
+﻿// Copyright Epic Games, Inc. All Rights Reserved.
+
+using UnrealBuildTool;
+
+public class NamingTokens : ModuleRules
+{
+	public NamingTokens(ReadOnlyTargetRules Target) : base(Target)
+	{
+		PCHUsage = ModuleRules.PCHUsageMode.UseExplicitOrSharedPCHs;
+		
+		PublicDependencyModuleNames.AddRange(
+			new string[]
+			{
+				"Core",
+				"EditorSubsystem"
+			}
+			);
+		
+		PrivateDependencyModuleNames.AddRange(
+			new string[]
+			{
+				"CoreUObject",
+				"Engine",
+				"Slate",
+				"SlateCore",
+				"InputCore",
+				"UnrealEd",
+				"BlueprintGraph",
+			}
+			);
+	}
+}
diff --git a/Engine/Plugins/Developer/NamingTokens/Source/NamingTokens/Private/Customization/NamingTokensCustomization.cpp b/Engine/Plugins/Developer/NamingTokens/Source/NamingTokens/Private/Customization/NamingTokensCustomization.cpp
new file mode 100644
--- /dev/null
+++ b/Engine/Plugins/Developer/NamingTokens/Source/NamingTokens/Private/Customization/NamingTokensCustomization.cpp
@@ -0,0 +1,233 @@
+// Copyright Epic Games, Inc. All Rights Reserved.
+
+#include "NamingTokensCustomization.h"
+
+#include "NamingTokenData.h"
+#include "NamingTokens.h"
+#include "Utils/NamingTokenUtils.h"
+
+#include "DetailWidgetRow.h"
+#include "Engine/Blueprint.h"
+#include "IDetailChildrenBuilder.h"
+#include "IDetailPropertyRow.h"
+#include "IPropertyUtilities.h"
+#include "Widgets/Input/SButton.h"
+#include "Widgets/Input/SComboBox.h"
+#include "Widgets/Text/STextBlock.h"
+
+#define LOCTEXT_NAMESPACE "NamingTokensCustomization"
+
+void FNamingTokensCustomization::CustomizeHeader(TSharedRef<IPropertyHandle> PropertyHandle, FDetailWidgetRow& HeaderRow,
+                                                 IPropertyTypeCustomizationUtils& CustomizationUtils)
+{
+	CustomizationUtilsPtr = &CustomizationUtils;
+
+	if (const UNamingTokens* OwningTokens = GetOwningNamingTokens())
+	{
+		OwningBlueprint = UBlueprint::GetBlueprintFromClass(OwningTokens->GetClass());
+		if (OwningBlueprint.IsValid())
+		{
+			FunctionNameHandle = PropertyHandle->GetChildHandle(GET_MEMBER_NAME_CHECKED(FNamingTokenData, FunctionName));
+			check(FunctionNameHandle.IsValid());
+
+			TokenKeyHandle = PropertyHandle->GetChildHandle(GET_MEMBER_NAME_CHECKED(FNamingTokenData, TokenKey));
+			check(TokenKeyHandle.IsValid());
+		}
+	}
+
+	// If owning blueprint is null, then we may not be customizing the naming tokens directly, such as a settings object.
+	// or the tokens are predefined and can't be extended in BP.
+	
+	HeaderRow
+	.NameContent()
+	[
+		PropertyHandle->CreatePropertyNameWidget()
+	]
+	.ValueContent()
+	[
+		PropertyHandle->CreatePropertyValueWidget(/* bDisplayDefaultPropertyButtons (adds identical options) */ false)
+	];
+}
+
+void FNamingTokensCustomization::CustomizeChildren(TSharedRef<IPropertyHandle> PropertyHandle, IDetailChildrenBuilder& ChildBuilder,
+	IPropertyTypeCustomizationUtils& CustomizationUtils)
+{
+	const bool bShouldApplyCustomization = OwningBlueprint.IsValid();
+
+	if (bShouldApplyCustomization)
+	{
+		const TArray<UFunction*> Functions = GetAvailableFunctions();
+		for (const UFunction* Function : Functions)
+		{
+			FunctionNames.Add(MakeShared<FString>(Function->GetName()));
+		}
+
+		FString CurrentValue;
+		FunctionNameHandle->GetValue(CurrentValue);
+
+		// Default selected function
+		for (const TSharedPtr<FString>& Name : FunctionNames)
+		{
+			if (*Name == CurrentValue)
+			{
+				SelectedFunctionName = Name;
+				break;
+			}
+		}
+
+		FunctionNameHandle->MarkHiddenByCustomization();
+	}
+	
+	uint32 NumChildren = 0;
+	PropertyHandle->GetNumChildren(NumChildren);
+	
+	for (uint32 ChildNum = 0; ChildNum < NumChildren; ++ChildNum)
+	{
+		const TSharedPtr<IPropertyHandle> ChildHandle = PropertyHandle->GetChildHandle(ChildNum);
+		if (!ChildHandle.IsValid() || (bShouldApplyCustomization && ChildHandle->GetProperty() == FunctionNameHandle->GetProperty()))
+		{
+			continue;
+		}
+		ChildBuilder.AddProperty(ChildHandle.ToSharedRef());
+	}
+
+	if (!bShouldApplyCustomization)
+	{
+		return;
+	}
+	
+	FDetailWidgetRow& HeaderRow = ChildBuilder.AddProperty(FunctionNameHandle.ToSharedRef()).CustomWidget();
+	HeaderRow.NameContent()
+	[
+		FunctionNameHandle->CreatePropertyNameWidget()
+	];
+	HeaderRow.ValueContent()
+	[
+		SNew(SHorizontalBox)
+		+SHorizontalBox::Slot()
+		[
+			SNew(SComboBox<TSharedPtr<FString>>)
+			.OptionsSource(&FunctionNames)
+			.OnGenerateWidget(this, &FNamingTokensCustomization::MakeComboBoxWidget)
+			.OnSelectionChanged(this, &FNamingTokensCustomization::OnFunctionSelected, FunctionNameHandle)
+			.InitiallySelectedItem(SelectedFunctionName)
+			[
+				SNew(STextBlock)
+				.Text(this, &FNamingTokensCustomization::GetSelectedFunctionText)
+			]
+		]
+		+ SHorizontalBox::Slot()
+		.AutoWidth()
+		.VAlign(VAlign_Center)
+		[
+			SNew(SButton)
+			[
+				SNew(SImage)
+				.Image(FAppStyle::Get().GetBrush("PListEditor.Button_AddToArray"))
+			]
+			.ButtonStyle(FAppStyle::Get(), "HoverHintOnly")
+			.ToolTipText(LOCTEXT("AddFunction", "Create and assign a new function graph for evaluating tokens."))
+			.OnClicked(this, &FNamingTokensCustomization::OnAddFunctionClicked)
+			.IsEnabled(this, &FNamingTokensCustomization::CanAddFunction)
+		]
+	];
+}
+
+UNamingTokens* FNamingTokensCustomization::GetOwningNamingTokens() const
+{
+	if (CustomizationUtilsPtr)
+	{
+		const TSharedPtr<IPropertyUtilities> PropertyUtilities = CustomizationUtilsPtr->GetPropertyUtilities();
+		if (PropertyUtilities.IsValid())
+		{
+			const TArray<TWeakObjectPtr<UObject>> ObjectsBeingCustomized = PropertyUtilities->GetSelectedObjects();
+
+			if (ObjectsBeingCustomized.Num() == 1)
+			{
+				if (UNamingTokens* Tokens = Cast<UNamingTokens>(ObjectsBeingCustomized[0]))
+				{
+					return Tokens;
+				}
+			}
+		}
+	}
+	return nullptr;
+}
+
+TArray<UFunction*> FNamingTokensCustomization::GetAvailableFunctions() const
+{
+	TArray<UFunction*> Result;
+	if (OwningBlueprint.IsValid())
+	{
+		for (TFieldIterator<UFunction> It(OwningBlueprint->SkeletonGeneratedClass ?
+			OwningBlueprint->SkeletonGeneratedClass : OwningBlueprint->GeneratedClass,
+			EFieldIteratorFlags::IncludeSuper); It; ++It)
+		{
+			if (UE::NamingTokens::Utils::Private::ValidateTokenFunction(*It))
+			{
+				Result.Add(*It);
+			}
+		}
+	}
+
+	return Result;
+}
+
+TSharedRef<SWidget> FNamingTokensCustomization::MakeComboBoxWidget(TSharedPtr<FString> InItem)
+{
+	return SNew(STextBlock)
+		.Text(FText::FromString(*InItem));
+}
+
+void FNamingTokensCustomization::OnFunctionSelected(TSharedPtr<FString> NewValue, ESelectInfo::Type SelectInfo,
+	TSharedPtr<IPropertyHandle> InFunctionNameHandle)
+{
+	if (NewValue.IsValid())
+	{
+		SelectedFunctionName = NewValue;
+		InFunctionNameHandle->SetValue(*NewValue);
+	}
+}
+
+FText FNamingTokensCustomization::GetSelectedFunctionText() const
+{
+	return SelectedFunctionName.IsValid() ? FText::FromString(*SelectedFunctionName) : FText::FromString(TEXT("Select Function"));
+}
+
+FReply FNamingTokensCustomization::OnAddFunctionClicked()
+{
+	if (UBlueprint* Blueprint = OwningBlueprint.Get())
+	{
+		FString TokenKey;
+		TokenKeyHandle->GetValue(TokenKey);
+
+		const FName NewFunctionName = UE::NamingTokens::Utils::Private::CreateNewTokenGraph(Blueprint, TokenKey);
+		FunctionNameHandle->SetValue(NewFunctionName);
+	}
+	
+	return FReply::Handled();
+}
+
+bool FNamingTokensCustomization::CanAddFunction() const
+{
+	FString TokenKey;
+	if (TokenKeyHandle.IsValid())
+	{
+		TokenKeyHandle->GetValue(TokenKey);
+		
+		const FString BaseFunctionName = UE::NamingTokens::Utils::Private::CreateBaseTokenFunctionName(TokenKey);
+
+		// Make sure the name isn't already being used.
+		for (const TSharedPtr<FString>& FunctionName : FunctionNames)
+		{
+			if (FunctionName.IsValid() && **FunctionName == BaseFunctionName)
+			{
+				return false;
+			}
+		}
+	}
+	
+	return !TokenKey.IsEmpty();
+}
+
+#undef LOCTEXT_NAMESPACE
diff --git a/Engine/Plugins/Developer/NamingTokens/Source/NamingTokens/Private/Customization/NamingTokensCustomization.h b/Engine/Plugins/Developer/NamingTokens/Source/NamingTokens/Private/Customization/NamingTokensCustomization.h
new file mode 100644
--- /dev/null
+++ b/Engine/Plugins/Developer/NamingTokens/Source/NamingTokens/Private/Customization/NamingTokensCustomization.h
@@ -0,0 +1,57 @@
+// Copyright Epic Games, Inc. All Rights Reserved.
+
+#pragma once
+
+#include "IPropertyTypeCustomization.h"
+#include "CoreMinimal.h"
+#include "Input/Reply.h"
+
+class UNamingTokens;
+class SWidget;
+
+class FNamingTokensCustomization final : public IPropertyTypeCustomization
+{
+public:
+	static TSharedRef<IPropertyTypeCustomization> MakeInstance()
+	{
+		return MakeShareable(new FNamingTokensCustomization());
+	}
+
+	virtual void CustomizeHeader(TSharedRef<IPropertyHandle> PropertyHandle, FDetailWidgetRow& HeaderRow, IPropertyTypeCustomizationUtils& CustomizationUtils) override;
+	virtual void CustomizeChildren(TSharedRef<IPropertyHandle> PropertyHandle, IDetailChildrenBuilder& ChildBuilder, IPropertyTypeCustomizationUtils& CustomizationUtils) override;
+
+	/** The naming tokens owning this customization. */
+	UNamingTokens* GetOwningNamingTokens() const;
+
+private:
+	/** Available functions which can be assigned to naming tokens. */
+	TArray<UFunction*> GetAvailableFunctions() const;
+	/** Make the combo box for selecting functions. */
+	TSharedRef<SWidget> MakeComboBoxWidget(TSharedPtr<FString> InItem);
+	/** When the user selects a function. */
+	void OnFunctionSelected(
+		TSharedPtr<FString> NewValue, 
+		ESelectInfo::Type SelectInfo, 
+		TSharedPtr<IPropertyHandle> InFunctionNameHandle);
+	/** The text of the currently selected function. */
+	FText GetSelectedFunctionText() const;
+	/** User clicked the button to add a new function. */
+	FReply OnAddFunctionClicked();
+	/** If a new function can be added. */
+	bool CanAddFunction() const;
+
+private:
+	/** Customization pointer. */
+	IPropertyTypeCustomizationUtils* CustomizationUtilsPtr = nullptr;
+	/** Blueprint of the naming tokens. */
+	TWeakObjectPtr<UBlueprint> OwningBlueprint;
+
+	/** All available function names. */
+	TArray<TSharedPtr<FString>> FunctionNames;
+	/** The selected function name. */
+	TSharedPtr<FString> SelectedFunctionName;
+	/** Property handle to the function name. */
+	TSharedPtr<IPropertyHandle> FunctionNameHandle;
+	/** Property handle to the token key. */
+	TSharedPtr<IPropertyHandle> TokenKeyHandle;
+};
\ No newline at end of file
diff --git a/Engine/Plugins/Developer/NamingTokens/Source/NamingTokens/Private/GlobalNamingTokens.cpp b/Engine/Plugins/Developer/NamingTokens/Source/NamingTokens/Private/GlobalNamingTokens.cpp
new file mode 100644
--- /dev/null
+++ b/Engine/Plugins/Developer/NamingTokens/Source/NamingTokens/Private/GlobalNamingTokens.cpp
@@ -0,0 +1,159 @@
+// Copyright Epic Games, Inc. All Rights Reserved.
+
+#include "GlobalNamingTokens.h"
+
+#include "NamingTokenContext.h"
+
+#include "Misc/App.h"
+
+#define LOCTEXT_NAMESPACE "GlobalNamingTokens"
+
+FString UGlobalNamingTokens::GlobalNamespace = TEXT("g");
+
+UGlobalNamingTokens::UGlobalNamingTokens()
+{
+	Namespace = GlobalNamespace;
+}
+
+void UGlobalNamingTokens::OnCreateDefaultTokens(TArray<FNamingTokenData>& Tokens)
+{
+	Tokens.Add({
+	LOCTEXT("ProjectName", "Project Name"),
+	TEXT("project"),
+		FTokenProcessorDelegateNative::CreateLambda([](const TObjectPtr<UNamingTokenContext>&Context) {
+			return FText::FromString(FApp::GetProjectName());
+		})
+	});
+	Tokens.Add({
+		LOCTEXT("Year4Digit", "Year (4 digit)"),
+		TEXT("yyyy"),
+		FTokenProcessorDelegateNative::CreateLambda([](const TObjectPtr<UNamingTokenContext>& Context)
+		{
+			return FText::FromString(FString::Printf(TEXT("%04d"), Context->CurrentDateTime.GetYear()));
+		})
+	});
+	Tokens.Add({
+		LOCTEXT("Year4Digit", "Year (2 digit)"),
+		TEXT("yy"),
+		FTokenProcessorDelegateNative::CreateLambda([](const TObjectPtr<UNamingTokenContext>& Context) {
+			return FText::FromString(FString::Printf(TEXT("%02d"), Context->CurrentDateTime.GetYear() % 100));
+		})
+	});
+	Tokens.Add({
+		LOCTEXT("MonthPascalCase", "3-character Month (Pascal Case)"),
+		TEXT("Mmm"),
+		FTokenProcessorDelegateNative::CreateLambda([](const TObjectPtr<UNamingTokenContext>& Context) {
+			const FString MonthName = Context->CurrentDateTime.ToFormattedString(TEXT("%b"));
+			return FText::FromString(MonthName.Left(1).ToUpper() + MonthName.Mid(1).ToLower());
+		})
+	});
+	Tokens.Add({
+		LOCTEXT("MonthUpperCase", "3-character Month (UPPERCASE)"),
+		TEXT("MMM"),
+		FTokenProcessorDelegateNative::CreateLambda([](const TObjectPtr<UNamingTokenContext>& Context) {
+			const FString MonthName = Context->CurrentDateTime.ToFormattedString(TEXT("%b"));
+			return FText::FromString(MonthName.ToUpper());
+		})
+	});
+	Tokens.Add({
+		LOCTEXT("MonthLowerCase", "3-character Month (lowercase)"),
+		TEXT("mmm"),
+		FTokenProcessorDelegateNative::CreateLambda([](const TObjectPtr<UNamingTokenContext>& Context) {
+			const FString MonthName = Context->CurrentDateTime.ToFormattedString(TEXT("%b"));
+			return FText::FromString(MonthName.ToLower());
+		})
+	});
+	Tokens.Add({
+		LOCTEXT("Month2Digit", "Month (2 digit)"),
+		TEXT("mm"),
+		FTokenProcessorDelegateNative::CreateLambda([](const TObjectPtr<UNamingTokenContext>& Context) {
+			return FText::FromString(FString::Printf(TEXT("%02d"), Context->CurrentDateTime.GetMonth()));
+		})
+	});
+	Tokens.Add({
+		LOCTEXT("DayPascalCase", "3-character Day (Pascal Case)"),
+		TEXT("Ddd"),
+		FTokenProcessorDelegateNative::CreateLambda([](const TObjectPtr<UNamingTokenContext>& Context) {
+			const FString DayName = Context->CurrentDateTime.ToFormattedString(TEXT("%a"));
+			return FText::FromString(DayName.Left(1).ToUpper() + DayName.Mid(1).ToLower());
+		})
+	});
+
+	Tokens.Add({
+		LOCTEXT("DayUpperCase", "3-character Day (UPPERCASE)"),
+		TEXT("DDD"),
+		FTokenProcessorDelegateNative::CreateLambda([](const TObjectPtr<UNamingTokenContext>& Context) {
+			const FString DayName = Context->CurrentDateTime.ToFormattedString(TEXT("%a"));
+			return FText::FromString(DayName.ToUpper());
+		})
+	});
+
+	Tokens.Add({
+		LOCTEXT("DayLowerCase", "3-character Day (lowercase)"),
+		TEXT("ddd"),
+		FTokenProcessorDelegateNative::CreateLambda([](const TObjectPtr<UNamingTokenContext>& Context) {
+			const FString DayName = Context->CurrentDateTime.ToFormattedString(TEXT("%a"));
+			return FText::FromString(DayName.ToLower());
+		})
+	});
+	Tokens.Add({
+		LOCTEXT("Day2Digit", "Day (2 digit)"),
+		TEXT("dd"),
+		FTokenProcessorDelegateNative::CreateLambda([](const TObjectPtr<UNamingTokenContext>&Context) {
+			return FText::FromString(FString::Printf(TEXT("%02d"), Context->CurrentDateTime.GetDay()));
+		})
+	});
+	Tokens.Add({
+	LOCTEXT("AMPMLowerCase", "am or pm (lowercase)"),
+	TEXT("ampm"),
+		FTokenProcessorDelegateNative::CreateLambda([](const TObjectPtr<UNamingTokenContext>&Context) {
+			const FString AMPM = Context->CurrentDateTime.ToFormattedString(TEXT("%P"));
+			return FText::FromString(AMPM);
+		})
+	});
+	Tokens.Add({
+	LOCTEXT("AMPMUpperCase", "AM or PM (UPPERCASE)"),
+	TEXT("AMPM"),
+		FTokenProcessorDelegateNative::CreateLambda([](const TObjectPtr<UNamingTokenContext>&Context) {
+			const FString AMPM = Context->CurrentDateTime.ToFormattedString(TEXT("%p"));
+			return FText::FromString(AMPM);
+		})
+	});
+	Tokens.Add({
+	LOCTEXT("Hour12", "Hour (12)"),
+	TEXT("12h"),
+		FTokenProcessorDelegateNative::CreateLambda([](const TObjectPtr<UNamingTokenContext>&Context) {
+			return FText::FromString(FString::Printf(TEXT("%02d"), Context->CurrentDateTime.GetHour12()));
+		})
+	});
+	Tokens.Add({
+	LOCTEXT("Hour24", "Hour (24)"),
+	TEXT("24h"),
+		FTokenProcessorDelegateNative::CreateLambda([](const TObjectPtr<UNamingTokenContext>&Context) {
+			return FText::FromString(FString::Printf(TEXT("%02d"), Context->CurrentDateTime.GetHour()));
+		})
+	});
+	Tokens.Add({
+	LOCTEXT("Minute", "Minute"),
+	TEXT("min"),
+		FTokenProcessorDelegateNative::CreateLambda([](const TObjectPtr<UNamingTokenContext>&Context) {
+			return FText::FromString(FString::Printf(TEXT("%02d"), Context->CurrentDateTime.GetMinute()));
+		})
+	});
+	Tokens.Add({
+	LOCTEXT("Second", "Second"),
+	TEXT("sec"),
+		FTokenProcessorDelegateNative::CreateLambda([](const TObjectPtr<UNamingTokenContext>&Context) {
+			return FText::FromString(FString::Printf(TEXT("%02d"), Context->CurrentDateTime.GetSecond()));
+		})
+	});
+	Tokens.Add({
+	LOCTEXT("Millisecond", "Millisecond"),
+	TEXT("ms"),
+		FTokenProcessorDelegateNative::CreateLambda([](const TObjectPtr<UNamingTokenContext>&Context) {
+			return FText::FromString(FString::Printf(TEXT("%02d"), Context->CurrentDateTime.GetMillisecond()));
+		})
+	});
+}
+
+#undef LOCTEXT_NAMESPACE
diff --git a/Engine/Plugins/Developer/NamingTokens/Source/NamingTokens/Private/NamingTokens.cpp b/Engine/Plugins/Developer/NamingTokens/Source/NamingTokens/Private/NamingTokens.cpp
new file mode 100644
--- /dev/null
+++ b/Engine/Plugins/Developer/NamingTokens/Source/NamingTokens/Private/NamingTokens.cpp
@@ -0,0 +1,223 @@
+// Copyright Epic Games, Inc. All Rights Reserved.
+
+#include "NamingTokens.h"
+
+#include "NamingTokenContext.h"
+#include "NamingTokensLog.h"
+#include "Utils/NamingTokenUtils.h"
+
+#include "Internationalization/Regex.h"
+#include "UObject/Package.h"
+
+#define LOCTEXT_NAMESPACE "NamingTokens"
+
+UNamingTokens::UNamingTokens()
+{
+	NamingTokenContextClass = UNamingTokenContext::StaticClass();
+}
+
+void UNamingTokens::PostInitProperties()
+{
+	UObject::PostInitProperties();
+	CreateDefaultTokens();
+}
+
+void UNamingTokens::CreateDefaultTokens()
+{
+	DefaultTokens.Empty();
+	OnCreateDefaultTokens(DefaultTokens);
+}
+
+#if WITH_EDITOR
+void UNamingTokens::PostEditChangeProperty(FPropertyChangedEvent& PropertyChangedEvent)
+{
+	Super::PostEditChangeProperty(PropertyChangedEvent);
+	if (PropertyChangedEvent.GetPropertyName() == GET_MEMBER_NAME_CHECKED(UNamingTokens, TestTokenInput))
+	{
+		EvaluateTestToken();
+	}
+}
+
+void UNamingTokens::EvaluateTestToken()
+{
+	const FNamingTokenResultData ResultData = EvaluateTokenString(TestTokenInput.ToString());
+	TestTokenResult = ResultData.EvaluatedText;
+}
+#endif
+
+void UNamingTokens::OnCreateDefaultTokens(TArray<FNamingTokenData>& Tokens)
+{
+}
+
+FNamingTokenResultData UNamingTokens::EvaluateTokenString(const FString& InTokenString)
+{
+	FNamingTokenResultData Result;
+	Result.OriginalString = InTokenString;
+
+	const FString RawTokenString = PreProcessTokenString(InTokenString);
+	FText FormattedText = FText::FromString(RawTokenString);
+
+	const TObjectPtr<UNamingTokenContext> Context = CreateContext();
+	
+	for (const FNamingTokenData& Token : GetDefaultTokens())
+	{
+		TryProcessToken(Token,Context, FormattedText);
+	}
+
+	for (const FNamingTokenData& Token : CustomTokens)
+	{
+		TryProcessToken(Token,Context, FormattedText);
+	}
+
+	Result.EvaluatedText = FormattedText;
+	FString FormattedString = FormattedText.ToString();
+
+	// Split folder path and file name.
+	int32 LastSlashIndex;
+	if (FormattedString.FindLastChar('/', LastSlashIndex))
+	{
+		Result.FolderPath = FormattedString.Left(LastSlashIndex);
+		if (Result.FolderPath.StartsWith("/"))
+		{
+			Result.FolderPath.RemoveFromStart("/");
+		}
+		Result.FileName = FormattedString.Mid(LastSlashIndex + 1);
+	}
+	else
+	{
+		Result.FolderPath = TEXT("");
+		Result.FileName = FormattedString;
+	}
+	
+	return Result;
+}
+
+FString UNamingTokens::CreateFormattedToken(const FNamingTokenData& InToken) const
+{
+	return FString::Printf(TEXT("{%s}"), *InToken.TokenKey);
+}
+
+FString UNamingTokens::GetFormattedTokensStringForDisplay() const
+{
+	FString FormattedTokensString;
+
+	for (const FNamingTokenData& Token : GetDefaultTokens())
+	{
+		FormattedTokensString += FString::Printf(TEXT("%s - %s\n"),
+			*CreateFormattedToken(Token), *Token.DisplayName.ToString());
+	}
+
+	for (const FNamingTokenData& Token : CustomTokens)
+	{
+		FormattedTokensString += FString::Printf(TEXT("%s - %s\n"),
+			*CreateFormattedToken(Token), *Token.DisplayName.ToString());
+	}
+	
+	return FormattedTokensString;
+}
+
+const TArray<FNamingTokenData>& UNamingTokens::GetDefaultTokens() const
+{
+	// For now we return the instance version, but we could potentially return the CDO's version.
+	return DefaultTokens;
+}
+
+UNamingTokenContext* UNamingTokens::CreateContext_Implementation() const
+{
+	const UClass* ContextClass = NamingTokenContextClass.LoadSynchronous();
+	if (!ContextClass)
+	{
+		ContextClass = UNamingTokenContext::StaticClass();
+	}
+	
+	TObjectPtr<UNamingTokenContext> Context = NewObject<UNamingTokenContext>(GetTransientPackage(), ContextClass);
+	Context->CurrentDateTime = FDateTime::Now();
+
+	return Context;
+}
+
+FString UNamingTokens::PreProcessTokenString(const FString& InTokenString) const
+{
+	return InTokenString;
+}
+
+FString UNamingTokens::RemoveNamespaceFromTokenString(const FString& InTokenString) const
+{
+	// Build the regex pattern to match {namespace.token}
+	if (Namespace.IsEmpty())
+	{
+		return InTokenString;
+	}
+	const FString PatternString = FString::Printf(TEXT(R"(\{\s*%s\.([a-zA-Z0-9_]+)\s*\})"), *Namespace);
+	const FRegexPattern Pattern(PatternString);
+
+	FRegexMatcher Matcher(Pattern, InTokenString);
+
+	FString ProcessedText = InTokenString;
+	while (Matcher.FindNext())
+	{
+		// Get the token part of {namespace.token}
+		FString Token = Matcher.GetCaptureGroup(1);
+
+		// Replace {namespace.token} with {token}
+		const FString FullMatch = Matcher.GetCaptureGroup(0);
+		ProcessedText = ProcessedText.Replace(*FullMatch, *FString::Printf(TEXT("{%s}"), *Token));
+	}
+
+	return ProcessedText;
+}
+
+bool UNamingTokens::TryProcessToken(const FNamingTokenData& InToken, UNamingTokenContext* InContext, FText& InOutFormattedText)
+{
+	const FString FormattedToken = *CreateFormattedToken(InToken);
+
+	FText ProcessedToken = FText::GetEmpty();
+
+	if (UFunction* BlueprintFunction = FindBlueprintFunctionForToken(InToken))
+	{
+		struct FFunctionParams
+		{
+			UNamingTokenContext* Input;
+			FText ReturnValue;
+		} Params;
+		Params.Input = InContext;
+		this->ProcessEvent(BlueprintFunction, &Params);
+		ProcessedToken = Params.ReturnValue;
+	}
+	else if (InToken.TokenProcessorNative.IsBound())
+	{
+		ProcessedToken = InToken.TokenProcessorNative.Execute(InContext);
+	}
+	else
+	{
+		UE_LOG(LogNamingTokens, Error, TEXT("Could not find function process for token %s"), *FormattedToken);
+		return false;
+	}
+		
+	FFormatNamedArguments Args;
+	Args.Add(InToken.TokenKey, ProcessedToken);
+	Args.Add(FString::Printf(TEXT("%s.%s"), *GetNamespace(), *InToken.TokenKey), ProcessedToken);
+
+	InOutFormattedText = FText::Format(InOutFormattedText, Args);
+	return true;
+}
+
+UFunction* UNamingTokens::FindBlueprintFunctionForToken(const FNamingTokenData& InTokenData) const
+{
+	const FName FunctionName = InTokenData.FunctionName;
+	if (UFunction* Function = GetClass()->FindFunctionByName(FunctionName))
+	{
+		if (UE::NamingTokens::Utils::Private::ValidateTokenFunction(Function))
+		{
+			return Function;
+		}
+		
+		UE_LOG(LogNamingTokens, Warning,
+			TEXT("Token processor function '%s' found for token '%s', but it doesn't have the correct signature"),
+			*FunctionName.ToString(), *InTokenData.TokenKey);
+	}
+
+	return nullptr;
+}
+
+#undef LOCTEXT_NAMESPACE
diff --git a/Engine/Plugins/Developer/NamingTokens/Source/NamingTokens/Private/NamingTokensEditorSubsystem.cpp b/Engine/Plugins/Developer/NamingTokens/Source/NamingTokens/Private/NamingTokensEditorSubsystem.cpp
new file mode 100644
--- /dev/null
+++ b/Engine/Plugins/Developer/NamingTokens/Source/NamingTokens/Private/NamingTokensEditorSubsystem.cpp
@@ -0,0 +1,190 @@
+// Copyright Epic Games, Inc. All Rights Reserved.
+
+#include "NamingTokensEditorSubsystem.h"
+
+#include "NamingTokens.h"
+
+#include "AssetRegistry/AssetData.h"
+#include "AssetRegistry/AssetDataTagMap.h"
+#include "AssetRegistry/AssetRegistryModule.h"
+#include "Engine/Blueprint.h"
+#include "Internationalization/Regex.h"
+#include "Misc/PackageName.h"
+
+UNamingTokens* UNamingTokensEditorSubsystem::GetNamingTokens(const FString& InNamespace) const
+{
+	auto LoadNamingToken = [this, &InNamespace](const TSoftClassPtr<UNamingTokens>& InTokensClass) -> UNamingTokens*
+	{
+		if (const UClass* Class = InTokensClass.LoadSynchronous())
+		{
+			UNamingTokens* NamingTokensCDO = Class->GetDefaultObject<UNamingTokens>();
+			if (ensure(NamingTokensCDO->GetNamespace() == InNamespace))
+			{
+				LoadedNamingTokens.Add(InNamespace, InTokensClass);
+			}
+			return NamingTokensCDO;
+		}
+		return nullptr;
+	};
+
+	// Check cache.
+	if (const TSoftClassPtr<UNamingTokens>* Tokens = LoadedNamingTokens.Find(InNamespace))
+	{
+		if (UNamingTokens* CachedTokens = LoadNamingToken(*Tokens))
+		{
+			return CachedTokens;
+		}
+	}
+
+	const UClass* TargetClass = UNamingTokens::StaticClass();
+	
+	// Check native classes.
+	{
+		constexpr EClassFlags InvalidClassFlags = CLASS_Deprecated | CLASS_NewerVersionExists | CLASS_Abstract;
+
+		TArray<UClass*> DerivedClasses;
+		GetDerivedClasses(TargetClass, DerivedClasses);
+		for (const UClass* DerivedClass : DerivedClasses)
+		{
+			if (DerivedClass->IsNative()
+				&& DerivedClass->IsChildOf(TargetClass)
+				&& !DerivedClass->HasAnyClassFlags(InvalidClassFlags)
+				&& !DerivedClass->GetName().StartsWith(TEXT("SKEL_"))
+				&& !DerivedClass->GetName().StartsWith(TEXT("REINST_"))
+				&& DerivedClass->GetDefaultObject<UNamingTokens>()->GetNamespace() == InNamespace)
+			{
+				return LoadNamingToken(DerivedClass);
+			}
+		}
+	}
+
+	// Check blueprint classes.
+	// NOTE: We may want to use OnAssetAdded and cache the asset or load the class then.  
+	{
+		const FAssetRegistryModule& AssetRegistryModule = FModuleManager::LoadModuleChecked<FAssetRegistryModule>("AssetRegistry");
+		const IAssetRegistry& AssetRegistry = AssetRegistryModule.Get();
+
+		FARFilter Filter;
+		Filter.ClassPaths.Add(UBlueprint::StaticClass()->GetClassPathName());
+		Filter.bRecursiveClasses = true;
+		Filter.bRecursivePaths = true;
+
+		TArray<FAssetData> BlueprintAssets;
+		AssetRegistry.GetAssets(Filter, BlueprintAssets);
+
+		TSet<FTopLevelAssetPath> DerivedClassNames;
+		AssetRegistry.GetDerivedClassNames({ TargetClass->GetClassPathName() }, {}, DerivedClassNames);
+
+		for (const FAssetData& AssetData : BlueprintAssets)
+		{
+			// Narrow down to only our assets.
+			const FAssetDataTagMapSharedView::FFindTagResult Result = AssetData.TagsAndValues.FindTag(TEXT("GeneratedClass"));
+			if (Result.IsSet())
+			{
+				const FString& GeneratedClassPathPtr = Result.GetValue();
+				const FTopLevelAssetPath ClassObjectPath(FPackageName::ExportTextPathToObjectPath(*GeneratedClassPathPtr));
+
+				if (DerivedClassNames.Contains(ClassObjectPath))
+				{
+					// Now check for the namespace match.
+					FString FoundNamespace;
+					if (AssetData.GetTagValue(UNamingTokens::GetNamespacePropertyName(), FoundNamespace))
+					{
+						const TSoftClassPtr<UNamingTokens> SoftClassPath(ClassObjectPath.ToString());
+						return LoadNamingToken(SoftClassPath);
+					}
+				}
+			}
+		}
+	}
+
+	return nullptr;
+}
+
+FNamingTokenResultData UNamingTokensEditorSubsystem::EvaluateTokenString(const FString& InTokenString, FNamingTokenFilterArgs InFilter)
+{
+	FNamingTokenResultData Result;
+
+	FText ProcessedTokenText = FText::FromString(InTokenString);
+	Result.EvaluatedText = ProcessedTokenText; // Set so we always have something in case no tokens were evaluated.
+
+	TSet<FString> Namespaces = GetNamingTokenNamespacesFromString(InTokenString, InFilter);
+	Namespaces.Append(InFilter.NamespacesToInclude); // Tokens could have been written without the namespace but should still be scoped to the filter.
+	
+	for (const FString& Namespace : Namespaces)
+	{
+		if (UNamingTokens* NamingTokens = GetNamingTokens(Namespace))
+		{
+			Result = NamingTokens->EvaluateTokenString(ProcessedTokenText.ToString());
+			ProcessedTokenText = Result.EvaluatedText;
+		}
+	}
+
+	Result.OriginalString = InTokenString;
+	
+	return Result;
+}
+
+FNamingTokenResultData UNamingTokensEditorSubsystem::EvaluateTokenString(const FString& InTokenString, const UNamingTokens* InNamingTokens, bool bIncludeGlobal)
+{
+	return EvaluateTokenString(InTokenString, { { InNamingTokens->GetNamespace() }, bIncludeGlobal });
+}
+
+void UNamingTokensEditorSubsystem::RegisterGlobalNamespace(const FString& InNamespace)
+{
+	GlobalNamespaces.Add(InNamespace);
+}
+
+void UNamingTokensEditorSubsystem::UnregisterGlobalNamespace(const FString& InNamespace)
+{
+	GlobalNamespaces.Remove(InNamespace);
+}
+
+TSet<FString> UNamingTokensEditorSubsystem::GetNamingTokenNamespacesFromString(const FString& InTokenString, const FNamingTokenFilterArgs& InFilter) const
+{
+	TSet<FString> Namespaces;
+
+	const FRegexPattern Pattern(R"(\{\s*([a-zA-Z0-9_]+)\.[a-zA-Z0-9_]+\s*\})");
+	FRegexMatcher Matcher(Pattern, InTokenString);
+	
+	while (Matcher.FindNext())
+	{
+		// Get the first capture group (namespace)
+		const FString Namespace = Matcher.GetCaptureGroup(1);
+
+		if (InFilter.NamespacesToInclude.IsEmpty() || InFilter.NamespacesToInclude.Contains(Namespace))
+		{
+			Namespaces.Add(Namespace);
+		}
+	}
+
+	// Add any global namespaces.
+	if (InFilter.bIncludeGlobal)
+	{
+		Namespaces.Append(GlobalNamespaces);
+	}
+	
+	return Namespaces;	
+}
+
+FString UNamingTokensEditorSubsystem::GetFormattedTokensStringForDisplay(const FNamingTokenFilterArgs& InFilter) const
+{
+	TSet<FString> Namespaces = TSet<FString>(InFilter.NamespacesToInclude);
+
+	if (InFilter.bIncludeGlobal)
+	{
+		Namespaces.Append(GlobalNamespaces.Array());
+	}
+
+	FString Result;
+
+	for (const FString& Namespace : Namespaces)
+	{
+		if (const UNamingTokens* NamingTokens = GetNamingTokens(Namespace))
+		{
+			Result += NamingTokens->GetFormattedTokensStringForDisplay();
+		}
+	}
+
+	return Result;
+}
diff --git a/Engine/Plugins/Developer/NamingTokens/Source/NamingTokens/Private/NamingTokensLog.cpp b/Engine/Plugins/Developer/NamingTokens/Source/NamingTokens/Private/NamingTokensLog.cpp
new file mode 100644
--- /dev/null
+++ b/Engine/Plugins/Developer/NamingTokens/Source/NamingTokens/Private/NamingTokensLog.cpp
@@ -0,0 +1,5 @@
+// Copyright Epic Games, Inc. All Rights Reserved.
+
+#include "NamingTokensLog.h"
+
+DEFINE_LOG_CATEGORY(LogNamingTokens);
diff --git a/Engine/Plugins/Developer/NamingTokens/Source/NamingTokens/Private/NamingTokensLog.h b/Engine/Plugins/Developer/NamingTokens/Source/NamingTokens/Private/NamingTokensLog.h
new file mode 100644
--- /dev/null
+++ b/Engine/Plugins/Developer/NamingTokens/Source/NamingTokens/Private/NamingTokensLog.h
@@ -0,0 +1,7 @@
+// Copyright Epic Games, Inc. All Rights Reserved.
+
+#pragma once
+
+#include "Logging/LogMacros.h"
+
+DECLARE_LOG_CATEGORY_EXTERN(LogNamingTokens, Log, All);
\ No newline at end of file
diff --git a/Engine/Plugins/Developer/NamingTokens/Source/NamingTokens/Private/NamingTokensModule.cpp b/Engine/Plugins/Developer/NamingTokens/Source/NamingTokens/Private/NamingTokensModule.cpp
new file mode 100644
--- /dev/null
+++ b/Engine/Plugins/Developer/NamingTokens/Source/NamingTokens/Private/NamingTokensModule.cpp
@@ -0,0 +1,47 @@
+﻿// Copyright Epic Games, Inc. All Rights Reserved.
+
+#include "NamingTokensModule.h"
+
+#include "Customization/NamingTokensCustomization.h"
+#include "GlobalNamingTokens.h"
+#include "NamingTokenData.h"
+#include "NamingTokensEditorSubsystem.h"
+
+#include "Editor.h"
+#include "PropertyEditorModule.h"
+
+#define LOCTEXT_NAMESPACE "FNamingTokensModule"
+
+void FNamingTokensModule::StartupModule()
+{
+	FPropertyEditorModule& PropertyModule = FModuleManager::GetModuleChecked<FPropertyEditorModule>("PropertyEditor");
+	PropertyModule.RegisterCustomPropertyTypeLayout(FNamingTokenData::StaticStruct()->GetFName(),
+		FOnGetPropertyTypeCustomizationInstance::CreateStatic(&FNamingTokensCustomization::MakeInstance));
+
+	FCoreDelegates::OnPostEngineInit.AddRaw(this, &FNamingTokensModule::OnPostEngineInit);
+}
+
+void FNamingTokensModule::ShutdownModule()
+{
+	FPropertyEditorModule& PropertyModule = FModuleManager::GetModuleChecked<FPropertyEditorModule>("PropertyEditor");
+	PropertyModule.UnregisterCustomPropertyTypeLayout(FNamingTokenData::StaticStruct()->GetFName());
+
+	if (GEditor)
+	{
+		GEditor->GetEditorSubsystem<UNamingTokensEditorSubsystem>()->UnregisterGlobalNamespace(UGlobalNamingTokens::GetGlobalNamespace());
+	}
+	
+	FCoreDelegates::OnPostEngineInit.RemoveAll(this);
+}
+
+void FNamingTokensModule::OnPostEngineInit()
+{
+	if (GEditor)
+	{
+		GEditor->GetEditorSubsystem<UNamingTokensEditorSubsystem>()->RegisterGlobalNamespace(UGlobalNamingTokens::GetGlobalNamespace());
+	}
+}
+
+#undef LOCTEXT_NAMESPACE
+	
+IMPLEMENT_MODULE(FNamingTokensModule, NamingTokens)
\ No newline at end of file
diff --git a/Engine/Plugins/Developer/NamingTokens/Source/NamingTokens/Private/NamingTokensModule.h b/Engine/Plugins/Developer/NamingTokens/Source/NamingTokens/Private/NamingTokensModule.h
new file mode 100644
--- /dev/null
+++ b/Engine/Plugins/Developer/NamingTokens/Source/NamingTokens/Private/NamingTokensModule.h
@@ -0,0 +1,16 @@
+﻿// Copyright Epic Games, Inc. All Rights Reserved.
+
+#pragma once
+
+#include "INamingTokensModule.h"
+
+class FNamingTokensModule : public INamingTokensModule
+{
+public:
+	/** IModuleInterface implementation */
+	virtual void StartupModule() override;
+	virtual void ShutdownModule() override;
+
+private:
+	void OnPostEngineInit();
+};
diff --git a/Engine/Plugins/Developer/NamingTokens/Source/NamingTokens/Private/Utils/NamingTokenUtils.cpp b/Engine/Plugins/Developer/NamingTokens/Source/NamingTokens/Private/Utils/NamingTokenUtils.cpp
new file mode 100644
--- /dev/null
+++ b/Engine/Plugins/Developer/NamingTokens/Source/NamingTokens/Private/Utils/NamingTokenUtils.cpp
@@ -0,0 +1,50 @@
+// Copyright Epic Games, Inc. All Rights Reserved.
+
+#include "NamingTokenUtils.h"
+
+#include "NamingTokens.h"
+
+#include "Kismet2/BlueprintEditorUtils.h"
+#include "Kismet2/KismetEditorUtilities.h"
+
+UFunction* UE::NamingTokens::Utils::Private::GetProcessTokenFunctionSignature()
+{
+	UFunction* FunctionSignature = UNamingTokens::StaticClass()->FindFunctionByName(UNamingTokens::GetProcessTokenTemplateFunctionName());
+	check(FunctionSignature);
+	return FunctionSignature;
+}
+
+bool UE::NamingTokens::Utils::Private::ValidateTokenFunction(const UFunction* InFunction)
+{
+	check(InFunction);
+	const UFunction* FunctionSignature = GetProcessTokenFunctionSignature();
+	return InFunction->IsSignatureCompatibleWith(FunctionSignature,
+				/* Ignore CPF_ReturnParm for inconsistency between native and BP return params. */
+				UFunction::GetDefaultIgnoredSignatureCompatibilityFlags() | CPF_ReturnParm)
+				&& InFunction->HasAnyFunctionFlags(FUNC_BlueprintCallable)
+				&& !InFunction->HasAnyFunctionFlags(FUNC_Private);
+}
+
+FString UE::NamingTokens::Utils::Private::CreateBaseTokenFunctionName(const FString& InTokenKey)
+{
+	return TEXT("ProcessToken_") + InTokenKey;
+}
+
+FName UE::NamingTokens::Utils::Private::CreateNewTokenGraph(UBlueprint* InBlueprint, const FString& InTokenKey)
+{
+	check(InBlueprint);
+	
+	const FString BaseName = CreateBaseTokenFunctionName(InTokenKey);
+	const FName FunctionName = FBlueprintEditorUtils::GenerateUniqueGraphName(InBlueprint, BaseName);
+
+	InBlueprint->Modify();
+		
+	UEdGraph* NewGraph = FBlueprintEditorUtils::CreateNewGraph(InBlueprint, FunctionName, UEdGraph::StaticClass(), UEdGraphSchema_K2::StaticClass());
+	
+	UFunction* FunctionSignature = GetProcessTokenFunctionSignature();
+		
+	FBlueprintEditorUtils::AddFunctionGraph<UFunction>(InBlueprint, NewGraph, /*bIsUserCreated=*/ true, FunctionSignature);
+	FKismetEditorUtilities::BringKismetToFocusAttentionOnObject(NewGraph);
+
+	return FunctionName;
+}
diff --git a/Engine/Plugins/Developer/NamingTokens/Source/NamingTokens/Private/Utils/NamingTokenUtils.h b/Engine/Plugins/Developer/NamingTokens/Source/NamingTokens/Private/Utils/NamingTokenUtils.h
new file mode 100644
--- /dev/null
+++ b/Engine/Plugins/Developer/NamingTokens/Source/NamingTokens/Private/Utils/NamingTokenUtils.h
@@ -0,0 +1,22 @@
+// Copyright Epic Games, Inc. All Rights Reserved.
+
+#pragma once
+
+#include "CoreMinimal.h"
+
+class UBlueprint;
+
+namespace UE::NamingTokens::Utils::Private
+{
+	/** Return the signature function all blueprint token processing should use. */
+	UFunction* GetProcessTokenFunctionSignature();
+
+	/** Validate that a token function can be used for token processing. */
+	bool ValidateTokenFunction(const UFunction* InFunction);
+
+	/** Create the default function token name to use, before unique naming adjustments. */
+	FString CreateBaseTokenFunctionName(const FString& InTokenKey);
+
+	/** Creates a new k2 graph in a blueprint for a given token key. */
+	FName CreateNewTokenGraph(UBlueprint* InBlueprint, const FString& InTokenKey);
+}
\ No newline at end of file
diff --git a/Engine/Plugins/Developer/NamingTokens/Source/NamingTokens/Public/GlobalNamingTokens.h b/Engine/Plugins/Developer/NamingTokens/Source/NamingTokens/Public/GlobalNamingTokens.h
new file mode 100644
--- /dev/null
+++ b/Engine/Plugins/Developer/NamingTokens/Source/NamingTokens/Public/GlobalNamingTokens.h
@@ -0,0 +1,33 @@
+// Copyright Epic Games, Inc. All Rights Reserved.
+
+#pragma once
+
+#include "NamingTokens.h"
+
+#include "GlobalNamingTokens.generated.h"
+
+class UNamingTokenContext;
+
+/**
+ * Default global tokens accessed project wide.
+ */
+UCLASS()
+class NAMINGTOKENS_API UGlobalNamingTokens final : public UNamingTokens
+{
+	GENERATED_BODY()
+
+public:
+	UGlobalNamingTokens();
+
+	static const FString& GetGlobalNamespace() { return GlobalNamespace; }
+	
+protected:
+	/** Define any default tokens. */
+	virtual void OnCreateDefaultTokens(TArray<FNamingTokenData>& Tokens) override;
+
+private:
+	/** The global namespace to register this naming token as. */
+	static FString GlobalNamespace;
+};
+
+
diff --git a/Engine/Plugins/Developer/NamingTokens/Source/NamingTokens/Public/INamingTokensModule.h b/Engine/Plugins/Developer/NamingTokens/Source/NamingTokens/Public/INamingTokensModule.h
new file mode 100644
--- /dev/null
+++ b/Engine/Plugins/Developer/NamingTokens/Source/NamingTokens/Public/INamingTokensModule.h
@@ -0,0 +1,26 @@
+﻿// Copyright Epic Games, Inc. All Rights Reserved.
+
+#pragma once
+
+#include "Modules/ModuleManager.h"
+
+class INamingTokensModule : public IModuleInterface
+{
+	/**
+	 * Singleton-like access to this module's interface.  This is just for convenience!
+	 * Beware of calling this during the shutdown phase, though.  Your module might have been unloaded already.
+	 *
+	 * @return Returns singleton instance, loading the module on demand if needed
+	 */
+	static INamingTokensModule& Get()
+	{
+		return FModuleManager::LoadModuleChecked<INamingTokensModule>("NamingTokens");
+	}
+
+	/**
+	 * Checks to see if this module is loaded and ready.  It is only valid to call Get() if IsAvailable() returns true.
+	 *
+	 * @return True if the module is loaded and ready to use
+	 */
+	static bool IsAvailable() { return FModuleManager::Get().IsModuleLoaded("NamingTokens"); }
+};
diff --git a/Engine/Plugins/Developer/NamingTokens/Source/NamingTokens/Public/NamingTokenContext.h b/Engine/Plugins/Developer/NamingTokens/Source/NamingTokens/Public/NamingTokenContext.h
new file mode 100644
--- /dev/null
+++ b/Engine/Plugins/Developer/NamingTokens/Source/NamingTokens/Public/NamingTokenContext.h
@@ -0,0 +1,21 @@
+// Copyright Epic Games, Inc. All Rights Reserved.
+
+#pragma once
+
+#include "CoreMinimal.h"
+
+#include "NamingTokenContext.generated.h"
+
+/**
+ * Context which gets passed to naming tokens during token evaluation.
+ */
+UCLASS(Blueprintable)
+class NAMINGTOKENS_API UNamingTokenContext : public UObject
+{
+	GENERATED_BODY()
+
+public:
+	/** The current date time. */
+	UPROPERTY(BlueprintReadOnly, Category = "NamingTokens")
+	FDateTime CurrentDateTime;
+};
\ No newline at end of file
diff --git a/Engine/Plugins/Developer/NamingTokens/Source/NamingTokens/Public/NamingTokenData.h b/Engine/Plugins/Developer/NamingTokens/Source/NamingTokens/Public/NamingTokenData.h
new file mode 100644
--- /dev/null
+++ b/Engine/Plugins/Developer/NamingTokens/Source/NamingTokens/Public/NamingTokenData.h
@@ -0,0 +1,67 @@
+// Copyright Epic Games, Inc. All Rights Reserved.
+
+#pragma once
+
+#include "CoreMinimal.h"
+#include "UObject/ObjectMacros.h"
+
+#include "NamingTokenData.generated.h"
+
+class UNamingTokenContext;
+
+/** Evaluated results from a template file string. */
+USTRUCT(BlueprintType)
+struct NAMINGTOKENS_API FNamingTokenResultData
+{
+	GENERATED_BODY()
+	
+	/** Original string without any modifications. */
+	UPROPERTY(BlueprintReadWrite, Category = "NamingTokens")
+	FString OriginalString;
+	/** The full string with evaluated tokens. */
+	UPROPERTY(BlueprintReadWrite, Category = "NamingTokens")
+	FText EvaluatedText;
+	/** The evaluated folder(s) without the first forward slash or filename. */
+	UPROPERTY(BlueprintReadWrite, Category = "NamingTokens")
+	FString FolderPath;
+	/** The evaluated filename. */
+	UPROPERTY(BlueprintReadWrite, Category = "NamingTokens")
+	FString FileName;
+};
+
+DECLARE_DELEGATE_RetVal_OneParam(FText, FTokenProcessorDelegateNative, const TObjectPtr<UNamingTokenContext>&);
+
+USTRUCT(BlueprintType)
+struct NAMINGTOKENS_API FNamingTokenData
+{
+	GENERATED_BODY()
+
+public:
+	FNamingTokenData() = default;
+	
+	FNamingTokenData(const FText& InTokenDisplayName, const FString& InTokenKey, const FTokenProcessorDelegateNative& InTokenProcessor)
+	{
+		DisplayName = InTokenDisplayName;
+		TokenKey = InTokenKey;
+		TokenProcessorNative = InTokenProcessor;
+	}
+
+	/** The key of the token to use. This is what the text must match in order to be evaluated. Brackets are automatically added and do not need to be included. */
+	UPROPERTY(BlueprintReadWrite, EditAnywhere, Category = "NamingTokens")
+	FString TokenKey;
+
+	/** The friendly display name of the token. */
+	UPROPERTY(BlueprintReadWrite, EditAnywhere, Category = "NamingTokens")
+	FText DisplayName;
+
+	/** A description of the token. */
+	UPROPERTY(BlueprintReadWrite, EditAnywhere, Category = "NamingTokens")
+	FText Description;
+
+	/** The function to use to evaluate the token. Only needed for blueprint implementations. */
+	UPROPERTY(EditAnywhere, Category = "NamingTokens", meta = (NoResetToDefault))
+	FName FunctionName;
+
+	/** The native delegate to execute to evaluate the function. If FunctionName is set then this is not used. */
+	FTokenProcessorDelegateNative TokenProcessorNative;
+};
diff --git a/Engine/Plugins/Developer/NamingTokens/Source/NamingTokens/Public/NamingTokens.h b/Engine/Plugins/Developer/NamingTokens/Source/NamingTokens/Public/NamingTokens.h
new file mode 100644
--- /dev/null
+++ b/Engine/Plugins/Developer/NamingTokens/Source/NamingTokens/Public/NamingTokens.h
@@ -0,0 +1,112 @@
+// Copyright Epic Games, Inc. All Rights Reserved.
+
+#pragma once
+
+#include "CoreMinimal.h"
+#include "UObject/ObjectMacros.h"
+#include "NamingTokenData.h"
+
+#include "NamingTokens.generated.h"
+
+class UNamingTokenContext;
+
+/**
+ * Subclass to define naming tokens to use for a specific tool or project.
+ */
+UCLASS(Blueprintable, Abstract)
+class NAMINGTOKENS_API UNamingTokens : public UObject
+{
+	GENERATED_BODY()
+
+public:
+	UNamingTokens();
+
+	virtual void PostInitProperties() override;
+
+	/** Create any default tokens. */
+	void CreateDefaultTokens();
+
+#if WITH_EDITOR
+	virtual void PostEditChangeProperty(FPropertyChangedEvent& PropertyChangedEvent) override;
+	/** Evaluate the test token input. */
+	void EvaluateTestToken();
+#endif
+
+protected:
+	/** Define any default tokens. */
+	virtual void OnCreateDefaultTokens(TArray<FNamingTokenData>& Tokens);
+
+public:
+	/** Evaluate a token string. */
+	virtual FNamingTokenResultData EvaluateTokenString(const FString& InTokenString);
+
+	/** Formats a token key into {key}. */
+	FString CreateFormattedToken(const FNamingTokenData& InToken) const;
+
+	/** Creates a friendly display string of all tokens. */
+	FString GetFormattedTokensStringForDisplay() const;
+
+	/** Retrieve the default tokens. */
+	const TArray<FNamingTokenData>& GetDefaultTokens() const;
+	
+	/** Retrieve the naming token's namespace. */
+	const FString& GetNamespace() const { return Namespace; }
+	/** The name of the namespace property. */
+	static FName GetNamespacePropertyName() { return GET_MEMBER_NAME_CHECKED(UNamingTokens, Namespace); }
+
+#if WITH_EDITOR
+public:
+	static FName GetProcessTokenTemplateFunctionName() { return GET_FUNCTION_NAME_CHECKED(UNamingTokens, ProcessTokenTemplateFunction); } 
+private:
+	/** Template function for us to dynamically create subclass graphs from matching this signature. */
+	UFUNCTION()
+	FText ProcessTokenTemplateFunction(UNamingTokenContext* InContext) { return FText::GetEmpty(); }
+#endif
+	
+protected:
+	/** Create the context to use for token evaluation. */
+	UFUNCTION(BlueprintNativeEvent, Category = "NamingTokens")
+	UNamingTokenContext* CreateContext() const;
+	virtual UNamingTokenContext* CreateContext_Implementation() const;
+
+	/** Run pre-processing logic on the token string prior to evaluation. */
+	virtual FString PreProcessTokenString(const FString& InTokenString) const;
+
+	/** Remove any uses of our namespace from the string. */
+	FString RemoveNamespaceFromTokenString(const FString& InTokenString) const;
+	
+	/** Process a token if is defined and return evaluated text. */
+	bool TryProcessToken(const FNamingTokenData& InToken, UNamingTokenContext* InContext, FText& InOutFormattedText);
+	
+	/** Return a blueprint function for a token processor if it exists. */
+	UFunction* FindBlueprintFunctionForToken(const FNamingTokenData& InTokenData) const;
+	
+private:
+	/** The default tokens defined by this class. */
+	TArray<FNamingTokenData> DefaultTokens;
+
+protected:
+	/** User defined tokens. */
+	UPROPERTY(EditDefaultsOnly, Category = "NamingTokens", meta = (TitleProperty="TokenKey"))
+	TArray<FNamingTokenData> CustomTokens;
+	
+	/** The namespace to identify this token. */
+	UPROPERTY(EditDefaultsOnly, AssetRegistrySearchable, Category = "NamingTokens")
+	FString Namespace;
+	
+	/** The class to use for the naming token context. */
+	UPROPERTY(EditDefaultsOnly, Category = "NamingTokens", NoClear)
+	TSoftClassPtr<UNamingTokenContext> NamingTokenContextClass;
+	
+#if WITH_EDITORONLY_DATA
+	/** Enter a sample string using your tokens to output an evaluated result to TestTokenResult. */
+	UPROPERTY(EditDefaultsOnly, Transient, Category = "NamingTokens")
+	FText TestTokenInput;
+
+	/** An evaluated text result of your token data. */
+	UPROPERTY(VisibleDefaultsOnly, Transient, Category = "NamingTokens")
+	FText TestTokenResult;
+#endif
+};
+
+
diff --git a/Engine/Plugins/Developer/NamingTokens/Source/NamingTokens/Public/NamingTokensEditorSubsystem.h b/Engine/Plugins/Developer/NamingTokens/Source/NamingTokens/Public/NamingTokensEditorSubsystem.h
new file mode 100644
--- /dev/null
+++ b/Engine/Plugins/Developer/NamingTokens/Source/NamingTokens/Public/NamingTokensEditorSubsystem.h
@@ -0,0 +1,69 @@
+// Copyright Epic Games, Inc. All Rights Reserved.
+
+#pragma once
+
+#include "EditorSubsystem.h"
+
+#include "NamingTokenData.h"
+
+#include "NamingTokensEditorSubsystem.generated.h"
+
+class UNamingTokens;
+
+USTRUCT(BlueprintType)
+struct FNamingTokenFilterArgs
+{
+	GENERATED_BODY()
+	
+	/** Namespaces to include in the search. */
+	UPROPERTY(EditAnywhere, BlueprintReadWrite, Category=NamingTokens)
+	TArray<FString> NamespacesToInclude;
+	/** Include global namespaces. */
+	UPROPERTY(EditAnywhere, BlueprintReadWrite, Category=NamingTokens)
+	bool bIncludeGlobal = true;
+};
+
+/**
+ * An editor subsystem for registering global tokens and evaluating strings across the entire project.
+ */
+UCLASS()
+class NAMINGTOKENS_API UNamingTokensEditorSubsystem : public UEditorSubsystem
+{
+	GENERATED_BODY()
+
+public:
+	/**
+	 * Lookup naming tokens given a namespace.
+	 * @return The found Naming Tokens object, or nullptr.
+	 */
+	UFUNCTION(BlueprintCallable, Category = "NamingTokens")
+	UNamingTokens* GetNamingTokens(const FString& InNamespace) const;
+
+	/** Parse and evaluate a token string. */
+	UFUNCTION(BlueprintCallable, Category = "NamingTokens")
+	FNamingTokenResultData EvaluateTokenString(const FString& InTokenString, FNamingTokenFilterArgs InFilter);
+
+	/** Evaluate a token string for a specific token instance. */
+	FNamingTokenResultData EvaluateTokenString(const FString& InTokenString, const UNamingTokens* InNamingTokens, bool bIncludeGlobal = true);
+
+	/** Register tokens as a global namespace. */
+	UFUNCTION(BlueprintCallable, Category = "NamingTokens")
+	void RegisterGlobalNamespace(const FString& InNamespace);
+
+	/** Remove tokens from a global namespace. */
+	UFUNCTION(BlueprintCallable, Category = "NamingTokens")
+	void UnregisterGlobalNamespace(const FString& InNamespace);
+	
+	/** Locate all reference naming token namespaces from a given string. */
+	TSet<FString> GetNamingTokenNamespacesFromString(const FString& InTokenString, const FNamingTokenFilterArgs& InFilter = FNamingTokenFilterArgs()) const;
+	
+	/** Creates a friendly display string of all tokens. */
+	FString GetFormattedTokensStringForDisplay(const FNamingTokenFilterArgs& InFilter) const;
+	
+private:
+	/** Naming tokens currently loaded from assets. */
+	mutable TMap<FString, TSoftClassPtr<UNamingTokens>> LoadedNamingTokens;
+
+	/** Namespaces considered global (don't need to include namespace to access). */
+	TSet<FString> GlobalNamespaces;
+};
