NamingTokens: Fix loctext conflict.

diff --git a/Engine/Plugins/Developer/NamingTokens/Source/NamingTokens/Private/GlobalNamingTokens.cpp b/Engine/Plugins/Developer/NamingTokens/Source/NamingTokens/Private/GlobalNamingTokens.cpp
--- a/Engine/Plugins/Developer/NamingTokens/Source/NamingTokens/Private/GlobalNamingTokens.cpp
+++ b/Engine/Plugins/Developer/NamingTokens/Source/NamingTokens/Private/GlobalNamingTokens.cpp
@@ -20,7 +20,7 @@ void UGlobalNamingTokens::OnCreateDefaultTokens(TArray<FNamingTokenData>& Tokens
 	Tokens.Add({
 	LOCTEXT("ProjectName", "Project Name"),
 	TEXT("project"),
-		FTokenProcessorDelegateNative::CreateLambda([](const TObjectPtr<UNamingTokenContext>&Context) {
+		FTokenProcessorDelegateNative::CreateLambda([](const TObjectPtr<UNamingTokenContext>& Context) {
 			return FText::FromString(FApp::GetProjectName());
 		})
 	});
@@ -33,7 +33,7 @@ void UGlobalNamingTokens::OnCreateDefaultTokens(TArray<FNamingTokenData>& Tokens
 		})
 	});
 	Tokens.Add({
-		LOCTEXT("Year4Digit", "Year (2 digit)"),
+		LOCTEXT("Year2Digit", "Year (2 digit)"),
 		TEXT("yy"),
 		FTokenProcessorDelegateNative::CreateLambda([](const TObjectPtr<UNamingTokenContext>& Context) {
 			return FText::FromString(FString::Printf(TEXT("%02d"), Context->CurrentDateTime.GetYear() % 100));
@@ -99,14 +99,14 @@ void UGlobalNamingTokens::OnCreateDefaultTokens(TArray<FNamingTokenData>& Tokens
 	Tokens.Add({
 		LOCTEXT("Day2Digit", "Day (2 digit)"),
 		TEXT("dd"),
-		FTokenProcessorDelegateNative::CreateLambda([](const TObjectPtr<UNamingTokenContext>&Context) {
+		FTokenProcessorDelegateNative::CreateLambda([](const TObjectPtr<UNamingTokenContext>& Context) {
 			return FText::FromString(FString::Printf(TEXT("%02d"), Context->CurrentDateTime.GetDay()));
 		})
 	});
 	Tokens.Add({
 	LOCTEXT("AMPMLowerCase", "am or pm (lowercase)"),
 	TEXT("ampm"),
-		FTokenProcessorDelegateNative::CreateLambda([](const TObjectPtr<UNamingTokenContext>&Context) {
+		FTokenProcessorDelegateNative::CreateLambda([](const TObjectPtr<UNamingTokenContext>& Context) {
 			const FString AMPM = Context->CurrentDateTime.ToFormattedString(TEXT("%P"));
 			return FText::FromString(AMPM);
 		})
@@ -114,7 +114,7 @@ void UGlobalNamingTokens::OnCreateDefaultTokens(TArray<FNamingTokenData>& Tokens
 	Tokens.Add({
 	LOCTEXT("AMPMUpperCase", "AM or PM (UPPERCASE)"),
 	TEXT("AMPM"),
-		FTokenProcessorDelegateNative::CreateLambda([](const TObjectPtr<UNamingTokenContext>&Context) {
+		FTokenProcessorDelegateNative::CreateLambda([](const TObjectPtr<UNamingTokenContext>& Context) {
 			const FString AMPM = Context->CurrentDateTime.ToFormattedString(TEXT("%p"));
 			return FText::FromString(AMPM);
 		})
@@ -122,35 +122,35 @@ void UGlobalNamingTokens::OnCreateDefaultTokens(TArray<FNamingTokenData>& Tokens
 	Tokens.Add({
 	LOCTEXT("Hour12", "Hour (12)"),
 	TEXT("12h"),
-		FTokenProcessorDelegateNative::CreateLambda([](const TObjectPtr<UNamingTokenContext>&Context) {
+		FTokenProcessorDelegateNative::CreateLambda([](const TObjectPtr<UNamingTokenContext>& Context) {
 			return FText::FromString(FString::Printf(TEXT("%02d"), Context->CurrentDateTime.GetHour12()));
 		})
 	});
 	Tokens.Add({
 	LOCTEXT("Hour24", "Hour (24)"),
 	TEXT("24h"),
-		FTokenProcessorDelegateNative::CreateLambda([](const TObjectPtr<UNamingTokenContext>&Context) {
+		FTokenProcessorDelegateNative::CreateLambda([](const TObjectPtr<UNamingTokenContext>& Context) {
 			return FText::FromString(FString::Printf(TEXT("%02d"), Context->CurrentDateTime.GetHour()));
 		})
 	});
 	Tokens.Add({
 	LOCTEXT("Minute", "Minute"),
 	TEXT("min"),
-		FTokenProcessorDelegateNative::CreateLambda([](const TObjectPtr<UNamingTokenContext>&Context) {
+		FTokenProcessorDelegateNative::CreateLambda([](const TObjectPtr<UNamingTokenContext>& Context) {
 			return FText::FromString(FString::Printf(TEXT("%02d"), Context->CurrentDateTime.GetMinute()));
 		})
 	});
 	Tokens.Add({
 	LOCTEXT("Second", "Second"),
 	TEXT("sec"),
-		FTokenProcessorDelegateNative::CreateLambda([](const TObjectPtr<UNamingTokenContext>&Context) {
+		FTokenProcessorDelegateNative::CreateLambda([](const TObjectPtr<UNamingTokenContext>& Context) {
 			return FText::FromString(FString::Printf(TEXT("%02d"), Context->CurrentDateTime.GetSecond()));
 		})
 	});
 	Tokens.Add({
 	LOCTEXT("Millisecond", "Millisecond"),
 	TEXT("ms"),
-		FTokenProcessorDelegateNative::CreateLambda([](const TObjectPtr<UNamingTokenContext>&Context) {
+		FTokenProcessorDelegateNative::CreateLambda([](const TObjectPtr<UNamingTokenContext>& Context) {
 			return FText::FromString(FString::Printf(TEXT("%02d"), Context->CurrentDateTime.GetMillisecond()));
 		})
 	});
