NamingTokens: Fix incorrect token namespace lookup of BP classes, preventing multiple token classes from being loaded properly. Change loaded token cache to use a new token instance rather than the CDO.

diff --git a/Engine/Plugins/Developer/NamingTokens/Source/NamingTokens/Private/NamingTokensEditorSubsystem.cpp b/Engine/Plugins/Developer/NamingTokens/Source/NamingTokens/Private/NamingTokensEditorSubsystem.cpp
--- a/Engine/Plugins/Developer/NamingTokens/Source/NamingTokens/Private/NamingTokensEditorSubsystem.cpp
+++ b/Engine/Plugins/Developer/NamingTokens/Source/NamingTokens/Private/NamingTokensEditorSubsystem.cpp
@@ -17,23 +17,20 @@ UNamingTokens* UNamingTokensEditorSubsystem::GetNamingTokens(const FString& InNa
 	{
 		if (const UClass* Class = InTokensClass.LoadSynchronous())
 		{
-			UNamingTokens* NamingTokensCDO = Class->GetDefaultObject<UNamingTokens>();
-			if (ensure(NamingTokensCDO->GetNamespace() == InNamespace))
+			UNamingTokens* NamingTokens = NewObject<UNamingTokens>(GetTransientPackage(), Class);
+			if (ensure(NamingTokens->GetNamespace() == InNamespace))
 			{
-				LoadedNamingTokens.Add(InNamespace, InTokensClass);
+				LoadedNamingTokens.Add(InNamespace, NamingTokens);
+				return NamingTokens;
 			}
-			return NamingTokensCDO;
 		}
 		return nullptr;
 	};
 
 	// Check cache.
-	if (const TSoftClassPtr<UNamingTokens>* Tokens = LoadedNamingTokens.Find(InNamespace))
+	if (const TObjectPtr<UNamingTokens>* Tokens = LoadedNamingTokens.Find(InNamespace))
 	{
-		if (UNamingTokens* CachedTokens = LoadNamingToken(*Tokens))
-		{
-			return CachedTokens;
-		}
+		return *Tokens;
 	}
 
 	const UClass* TargetClass = UNamingTokens::StaticClass();
@@ -88,7 +85,8 @@ UNamingTokens* UNamingTokensEditorSubsystem::GetNamingTokens(const FString& InNa
 				{
 					// Now check for the namespace match.
 					FString FoundNamespace;
-					if (AssetData.GetTagValue(UNamingTokens::GetNamespacePropertyName(), FoundNamespace))
+					if (AssetData.GetTagValue(UNamingTokens::GetNamespacePropertyName(), FoundNamespace)
+						&& FoundNamespace == InNamespace)
 					{
 						const TSoftClassPtr<UNamingTokens> SoftClassPath(ClassObjectPath.ToString());
 						return LoadNamingToken(SoftClassPath);
diff --git a/Engine/Plugins/Developer/NamingTokens/Source/NamingTokens/Public/NamingTokensEditorSubsystem.h b/Engine/Plugins/Developer/NamingTokens/Source/NamingTokens/Public/NamingTokensEditorSubsystem.h
--- a/Engine/Plugins/Developer/NamingTokens/Source/NamingTokens/Public/NamingTokensEditorSubsystem.h
+++ b/Engine/Plugins/Developer/NamingTokens/Source/NamingTokens/Public/NamingTokensEditorSubsystem.h
@@ -62,7 +62,8 @@ public:
 	
 private:
 	/** Naming tokens currently loaded from assets. */
-	mutable TMap<FString, TSoftClassPtr<UNamingTokens>> LoadedNamingTokens;
+	UPROPERTY(Transient)
+	mutable TMap<FString, TObjectPtr<UNamingTokens>> LoadedNamingTokens;
 
 	/** Namespaces considered global (don't need to include namespace to access). */
 	TSet<FString> GlobalNamespaces;
