NamingTokens: Add checks for duplicate namespaces.

Add error log when trying to register a global namespace which is already registered.

GetNamingTokens now checks for and logs warnings for duplicate namespaces.
GetNamingTokens now selects a BP child of a native class with the same namespace, instead of the native class.

Add BP function to retrieve multiple token objects from multiple namespaces.
Add BP functions for retrieving and checking namespaces.

Allow GlobalNamingTokens to be extended in C++ since it can be extended in BP, and its possible a user may want to extend, then register it as a new global namespace.

Make TokenResultData BP readonly as there isn't a reason to modify in BP.

diff --git a/Engine/Plugins/Developer/NamingTokens/Source/NamingTokens/Private/NamingTokensEditorSubsystem.cpp b/Engine/Plugins/Developer/NamingTokens/Source/NamingTokens/Private/NamingTokensEditorSubsystem.cpp
--- a/Engine/Plugins/Developer/NamingTokens/Source/NamingTokens/Private/NamingTokensEditorSubsystem.cpp
+++ b/Engine/Plugins/Developer/NamingTokens/Source/NamingTokens/Private/NamingTokensEditorSubsystem.cpp
@@ -3,6 +3,7 @@
 #include "NamingTokensEditorSubsystem.h"
 
 #include "NamingTokens.h"
+#include "NamingTokensLog.h"
 
 #include "AssetRegistry/AssetData.h"
 #include "AssetRegistry/AssetDataTagMap.h"
@@ -20,7 +21,7 @@ UNamingTokens* UNamingTokensEditorSubsystem::GetNamingTokens(const FString& InNa
 			UNamingTokens* NamingTokens = NewObject<UNamingTokens>(GetTransientPackage(), Class);
 			if (ensure(NamingTokens->GetNamespace() == InNamespace))
 			{
-				LoadedNamingTokens.Add(InNamespace, NamingTokens);
+				CachedNamingTokens.Add(InNamespace, NamingTokens);
 				return NamingTokens;
 			}
 		}
@@ -28,11 +29,13 @@ UNamingTokens* UNamingTokensEditorSubsystem::GetNamingTokens(const FString& InNa
 	};
 
 	// Check cache.
-	if (const TObjectPtr<UNamingTokens>* Tokens = LoadedNamingTokens.Find(InNamespace))
+	if (const TObjectPtr<UNamingTokens>* Tokens = CachedNamingTokens.Find(InNamespace))
 	{
 		return *Tokens;
 	}
 
+	UNamingTokens* FoundTokens = nullptr;
+
 	const UClass* TargetClass = UNamingTokens::StaticClass();
 	
 	// Check native classes.
@@ -50,7 +53,15 @@ UNamingTokens* UNamingTokensEditorSubsystem::GetNamingTokens(const FString& InNa
 				&& !DerivedClass->GetName().StartsWith(TEXT("REINST_"))
 				&& DerivedClass->GetDefaultObject<UNamingTokens>()->GetNamespace() == InNamespace)
 			{
-				return LoadNamingToken(DerivedClass);
+				if (!FoundTokens)
+				{
+					FoundTokens = LoadNamingToken(DerivedClass);
+				}
+				else
+				{
+					UE_LOG(LogNamingTokens, Warning, TEXT("Namespace '%s' exists more than once in native class '%s' and '%s'."),
+						*InNamespace, *FoundTokens->GetClass()->GetName(), *DerivedClass->GetName());
+				}
 			}
 		}
 	}
@@ -88,15 +99,59 @@ UNamingTokens* UNamingTokensEditorSubsystem::GetNamingTokens(const FString& InNa
 					if (AssetData.GetTagValue(UNamingTokens::GetNamespacePropertyName(), FoundNamespace)
 						&& FoundNamespace == InNamespace)
 					{
-						const TSoftClassPtr<UNamingTokens> SoftClassPath(ClassObjectPath.ToString());
-						return LoadNamingToken(SoftClassPath);
+						// Allow BP children to overwrite native classes.
+						auto IsAssetChildOfNativeToken = [&AssetRegistry, &InNamespace](const FTopLevelAssetPath& InAssetPath, const UNamingTokens* InParentTokens) -> bool
+						{
+							check(InParentTokens);
+							if (InParentTokens->GetClass()->IsNative())
+							{
+								// We need to make sure this BP is a child of the chosen native class.
+								TSet<FTopLevelAssetPath> NativeDerivedClassNames;
+								AssetRegistry.GetDerivedClassNames({ InParentTokens->GetClass()->GetClassPathName() },
+									{}, NativeDerivedClassNames);
+
+								if (NativeDerivedClassNames.Contains(InAssetPath))
+								{
+									UE_LOG(LogNamingTokens, Log, TEXT("Using namespace '%s' of BP child '%s' instead of native parent '%s'."),
+										*InNamespace, *InAssetPath.ToString(), *InParentTokens->GetClass()->GetName());
+									return true;
+								}
+							}
+							return false;
+						};
+						
+						if (!FoundTokens || IsAssetChildOfNativeToken(ClassObjectPath, FoundTokens))
+						{
+							const TSoftClassPtr<UNamingTokens> SoftClassPath(ClassObjectPath.ToString());
+							FoundTokens = LoadNamingToken(SoftClassPath);
+						}
+						else
+						{
+							UE_LOG(LogNamingTokens, Warning, TEXT("Namespace '%s' exists more than once in class '%s' and BP asset '%s'."),
+								*InNamespace, *FoundTokens->GetClass()->GetName(), *AssetData.PackageName.ToString());
+						}
 					}
 				}
 			}
 		}
 	}
 
-	return nullptr;
+	return FoundTokens;
+}
+
+TArray<UNamingTokens*> UNamingTokensEditorSubsystem::GetMultipleNamingTokens(const TArray<FString>& InNamespaces) const
+{
+	TArray<UNamingTokens*> Result;
+	const TSet<FString> UniqueNamespaces(InNamespaces);
+	for (const FString& Namespace : UniqueNamespaces)
+	{
+		if (UNamingTokens* NamingTokens = GetNamingTokens(Namespace))
+		{
+			Result.Add(NamingTokens);
+		}
+	}
+
+	return Result;
 }
 
 FNamingTokenResultData UNamingTokensEditorSubsystem::EvaluateTokenString(const FString& InTokenString, FNamingTokenFilterArgs InFilter)
@@ -130,6 +185,13 @@ FNamingTokenResultData UNamingTokensEditorSubsystem::EvaluateTokenString(const F
 
 void UNamingTokensEditorSubsystem::RegisterGlobalNamespace(const FString& InNamespace)
 {
+	if (IsGlobalNamespaceRegistered(InNamespace))
+	{
+		const FString UnregisterFunctionName = GET_FUNCTION_NAME_CHECKED(UNamingTokensEditorSubsystem, UnregisterGlobalNamespace).ToString();
+		UE_LOG(LogNamingTokens, Error, TEXT("NamingTokens Namespace '%s' is already registered as a global namespace. Call '%s' first if you wish to overwrite the namespace."),
+			 *InNamespace, *UnregisterFunctionName);
+		return;
+	}
 	GlobalNamespaces.Add(InNamespace);
 }
 
@@ -138,6 +200,16 @@ void UNamingTokensEditorSubsystem::UnregisterGlobalNamespace(const FString& InNa
 	GlobalNamespaces.Remove(InNamespace);
 }
 
+bool UNamingTokensEditorSubsystem::IsGlobalNamespaceRegistered(const FString& InNamespace) const
+{
+	return GlobalNamespaces.Contains(InNamespace);
+}
+
+TArray<FString> UNamingTokensEditorSubsystem::GetGlobalNamespaces() const
+{
+	return GlobalNamespaces.Array();
+}
+
 TSet<FString> UNamingTokensEditorSubsystem::GetNamingTokenNamespacesFromString(const FString& InTokenString, const FNamingTokenFilterArgs& InFilter) const
 {
 	TSet<FString> Namespaces;
@@ -186,3 +258,8 @@ FString UNamingTokensEditorSubsystem::GetFormattedTokensStringForDisplay(const F
 
 	return Result;
 }
+
+void UNamingTokensEditorSubsystem::ClearCachedNamingTokens()
+{
+	CachedNamingTokens.Empty();
+}
diff --git a/Engine/Plugins/Developer/NamingTokens/Source/NamingTokens/Private/NamingTokensModule.cpp b/Engine/Plugins/Developer/NamingTokens/Source/NamingTokens/Private/NamingTokensModule.cpp
--- a/Engine/Plugins/Developer/NamingTokens/Source/NamingTokens/Private/NamingTokensModule.cpp
+++ b/Engine/Plugins/Developer/NamingTokens/Source/NamingTokens/Private/NamingTokensModule.cpp
@@ -38,7 +38,13 @@ void FNamingTokensModule::OnPostEngineInit()
 {
 	if (GEditor)
 	{
-		GEditor->GetEditorSubsystem<UNamingTokensEditorSubsystem>()->RegisterGlobalNamespace(UGlobalNamingTokens::GetGlobalNamespace());
+		if (UNamingTokensEditorSubsystem* NamingTokensSubsystem = GEditor->GetEditorSubsystem<UNamingTokensEditorSubsystem>())
+		{
+			if (!NamingTokensSubsystem->IsGlobalNamespaceRegistered(UGlobalNamingTokens::GetGlobalNamespace()))
+			{
+				NamingTokensSubsystem->RegisterGlobalNamespace(UGlobalNamingTokens::GetGlobalNamespace());
+			}
+		}
 	}
 }
 
diff --git a/Engine/Plugins/Developer/NamingTokens/Source/NamingTokens/Public/GlobalNamingTokens.h b/Engine/Plugins/Developer/NamingTokens/Source/NamingTokens/Public/GlobalNamingTokens.h
--- a/Engine/Plugins/Developer/NamingTokens/Source/NamingTokens/Public/GlobalNamingTokens.h
+++ b/Engine/Plugins/Developer/NamingTokens/Source/NamingTokens/Public/GlobalNamingTokens.h
@@ -12,7 +12,7 @@ class UNamingTokenContext;
  * Default global tokens accessed project wide.
  */
 UCLASS()
-class NAMINGTOKENS_API UGlobalNamingTokens final : public UNamingTokens
+class NAMINGTOKENS_API UGlobalNamingTokens : public UNamingTokens
 {
 	GENERATED_BODY()
 
@@ -29,5 +29,3 @@ private:
 	/** The global namespace to register this naming token as. */
 	static FString GlobalNamespace;
 };
-
-
diff --git a/Engine/Plugins/Developer/NamingTokens/Source/NamingTokens/Public/NamingTokenData.h b/Engine/Plugins/Developer/NamingTokens/Source/NamingTokens/Public/NamingTokenData.h
--- a/Engine/Plugins/Developer/NamingTokens/Source/NamingTokens/Public/NamingTokenData.h
+++ b/Engine/Plugins/Developer/NamingTokens/Source/NamingTokens/Public/NamingTokenData.h
@@ -16,16 +16,16 @@ struct NAMINGTOKENS_API FNamingTokenResultData
 	GENERATED_BODY()
 	
 	/** Original string without any modifications. */
-	UPROPERTY(BlueprintReadWrite, Category = "NamingTokens")
+	UPROPERTY(BlueprintReadOnly, Category = "NamingTokens")
 	FString OriginalString;
 	/** The full string with evaluated tokens. */
-	UPROPERTY(BlueprintReadWrite, Category = "NamingTokens")
+	UPROPERTY(BlueprintReadOnly, Category = "NamingTokens")
 	FText EvaluatedText;
 	/** The evaluated folder(s) without the first forward slash or filename. */
-	UPROPERTY(BlueprintReadWrite, Category = "NamingTokens")
+	UPROPERTY(BlueprintReadOnly, Category = "NamingTokens")
 	FString FolderPath;
 	/** The evaluated filename. */
-	UPROPERTY(BlueprintReadWrite, Category = "NamingTokens")
+	UPROPERTY(BlueprintReadOnly, Category = "NamingTokens")
 	FString FileName;
 };
 
diff --git a/Engine/Plugins/Developer/NamingTokens/Source/NamingTokens/Public/NamingTokensEditorSubsystem.h b/Engine/Plugins/Developer/NamingTokens/Source/NamingTokens/Public/NamingTokensEditorSubsystem.h
--- a/Engine/Plugins/Developer/NamingTokens/Source/NamingTokens/Public/NamingTokensEditorSubsystem.h
+++ b/Engine/Plugins/Developer/NamingTokens/Source/NamingTokens/Public/NamingTokensEditorSubsystem.h
@@ -33,12 +33,25 @@ class NAMINGTOKENS_API UNamingTokensEditorSubsystem : public UEditorSubsystem
 
 public:
 	/**
-	 * Lookup naming tokens given a namespace.
+	 * Lookup naming tokens given a namespace. This will look first in cached naming tokens,
+	 * then native classes, then blueprint classes. Call ClearCachedNamingTokens to reset the cache
+	 * and force a full lookup.
+	 *
+	 * @param InNamespace The namespace of the tokens.
 	 * @return The found Naming Tokens object, or nullptr.
 	 */
 	UFUNCTION(BlueprintCallable, Category = "NamingTokens")
 	UNamingTokens* GetNamingTokens(const FString& InNamespace) const;
 
+	/**
+	 * Lookup multiple naming tokens from multiple namespaces.
+	 *
+	 * @param InNamespaces An array of all token namespaces.
+	 * @return An array of matching tokens objects.
+	 */
+	UFUNCTION(BlueprintCallable, Category = "NamingTokens")
+	TArray<UNamingTokens*> GetMultipleNamingTokens(const TArray<FString>& InNamespaces) const;
+
 	/** Parse and evaluate a token string. */
 	UFUNCTION(BlueprintCallable, Category = "NamingTokens")
 	FNamingTokenResultData EvaluateTokenString(const FString& InTokenString, FNamingTokenFilterArgs InFilter);
@@ -46,24 +59,45 @@ public:
 	/** Evaluate a token string for a specific token instance. */
 	FNamingTokenResultData EvaluateTokenString(const FString& InTokenString, const UNamingTokens* InNamingTokens, bool bIncludeGlobal = true);
 
-	/** Register tokens as a global namespace. */
+	/**
+	 * Register tokens as a global namespace. This prevents the need to include the namespace in a token string.
+	 * @param InNamespace The namespace of the tokens to register.
+	 */
 	UFUNCTION(BlueprintCallable, Category = "NamingTokens")
 	void RegisterGlobalNamespace(const FString& InNamespace);
 
-	/** Remove tokens from a global namespace. */
+	/**
+	 * Remove tokens from a global namespace.
+	 * @param InNamespace The namespace of the tokens to unregister.
+	 */
 	UFUNCTION(BlueprintCallable, Category = "NamingTokens")
 	void UnregisterGlobalNamespace(const FString& InNamespace);
+
+	/** Checks if a namespace is registered globally. */
+	UFUNCTION(BlueprintCallable, Category = "NamingTokens")
+	bool IsGlobalNamespaceRegistered(const FString& InNamespace) const;
 	
+	/** Retrieve the registered global namespaces. */
+	UFUNCTION(BlueprintCallable, Category = "NamingTokens")
+	TArray<FString> GetGlobalNamespaces() const;
+
 	/** Locate all reference naming token namespaces from a given string. */
 	TSet<FString> GetNamingTokenNamespacesFromString(const FString& InTokenString, const FNamingTokenFilterArgs& InFilter = FNamingTokenFilterArgs()) const;
 	
 	/** Creates a friendly display string of all tokens. */
 	FString GetFormattedTokensStringForDisplay(const FNamingTokenFilterArgs& InFilter) const;
+
+	/**
+	 * Clears all currently cached naming tokens. They will be loaded on demand when required.
+	 * This can avoid an editor restart if a namespace is adjusted on a BP asset, and you don't want the old namespace to access it.
+	 */
+	UFUNCTION(BlueprintCallable, Category = "NamingTokens")
+	void ClearCachedNamingTokens();
 	
 private:
 	/** Naming tokens currently loaded from assets. */
 	UPROPERTY(Transient)
-	mutable TMap<FString, TObjectPtr<UNamingTokens>> LoadedNamingTokens;
+	mutable TMap<FString, TObjectPtr<UNamingTokens>> CachedNamingTokens;
 
 	/** Namespaces considered global (don't need to include namespace to access). */
 	TSet<FString> GlobalNamespaces;
