TakeRecorder: Support NamingTokens. Take recorder specific names are still supported in addition to available NamingTokens namespaces.

NamingTokens:
- The context has been removed and instead an OnPreEvaluate function has been added. PreEvaluate takes a struct with shared data (only including a time stamp currently) which is consistent across namespaces during a run. Instead of utilizing a context, we use OnPreEvaluate to calculate and store any data directly in the Naming Tokens blueprint which can be accessed during token evaluation later.
- Any existing BP token functions will need to have their input parameter removed to function again.
- Added OnPostEvaluate in case any data should be cleaned up, since the NamingTokens object will live on, unlike the former context.
- Added EvaluateTokenText to support token evaluation with localization.
- Cleanup and consolidation of some functions.

LiveLinkHub: Updated use of Naming Tokens.

diff --git a/Engine/Plugins/Animation/LiveLinkHub/Source/LiveLinkHub/Private/Config/LiveLinkHubTemplateTokens.cpp b/Engine/Plugins/Animation/LiveLinkHub/Source/LiveLinkHub/Private/Config/LiveLinkHubTemplateTokens.cpp
--- a/Engine/Plugins/Animation/LiveLinkHub/Source/LiveLinkHub/Private/Config/LiveLinkHubTemplateTokens.cpp
+++ b/Engine/Plugins/Animation/LiveLinkHub/Source/LiveLinkHub/Private/Config/LiveLinkHubTemplateTokens.cpp
@@ -9,7 +9,6 @@
 ULiveLinkHubNamingTokens::ULiveLinkHubNamingTokens()
 {
 	Namespace = TEXT("llh");
-	NamingTokenContextClass = ULiveLinkHubNamingTokenContext::StaticClass();
 }
 
 void ULiveLinkHubNamingTokens::OnCreateDefaultTokens(TArray<FNamingTokenData>& Tokens)
@@ -17,24 +16,27 @@ void ULiveLinkHubNamingTokens::OnCreateDefaultTokens(TArray<FNamingTokenData>& T
 	Super::OnCreateDefaultTokens(Tokens);
 
 	Tokens.Add({
-		LOCTEXT("SessionName", "Session Name"),
 		TEXT("session"),
-		FTokenProcessorDelegateNative::CreateLambda([](const TObjectPtr<UNamingTokenContext>& Context)
-		{
-			return FText::FromString(FString::Printf(TEXT("%s"), *CastChecked<ULiveLinkHubNamingTokenContext>(Context)->SessionName));
+		LOCTEXT("SessionName", "Session Name"),
+		FNamingTokenData::FTokenProcessorDelegateNative::CreateLambda([this]() {
+			return FText::FromString(FString::Printf(TEXT("%s"), *SessionName));
 		})
 	});
 }
 
-UNamingTokenContext* ULiveLinkHubNamingTokens::CreateContext_Implementation() const
+void ULiveLinkHubNamingTokens::OnPreEvaluate_Implementation(const FNamingTokensEvaluationData& InEvaluationData)
 {
-	ULiveLinkHubNamingTokenContext* Context = CastChecked<ULiveLinkHubNamingTokenContext>(Super::CreateContext_Implementation());
+	Super::OnPreEvaluate_Implementation(InEvaluationData);
 	if (const TSharedPtr<ILiveLinkHubSessionManager> SessionManager = FModuleManager::Get().GetModuleChecked<FLiveLinkHubModule>("LiveLinkHub").GetLiveLinkHub()->GetSessionManager())
 	{
-		Context->SessionName = FPaths::GetBaseFilename(SessionManager->GetLastConfigPath());
+		SessionName = FPaths::GetBaseFilename(SessionManager->GetLastConfigPath());
 	}
-	
-	return Context;
 }
 
-#undef LOCTEXT_NAMESPACE
\ No newline at end of file
+void ULiveLinkHubNamingTokens::OnPostEvaluate_Implementation()
+{
+	Super::OnPostEvaluate_Implementation();
+	SessionName.Empty();
+}
+
+#undef LOCTEXT_NAMESPACE
diff --git a/Engine/Plugins/Animation/LiveLinkHub/Source/LiveLinkHub/Private/Config/LiveLinkHubTemplateTokens.h b/Engine/Plugins/Animation/LiveLinkHub/Source/LiveLinkHub/Private/Config/LiveLinkHubTemplateTokens.h
--- a/Engine/Plugins/Animation/LiveLinkHub/Source/LiveLinkHub/Private/Config/LiveLinkHubTemplateTokens.h
+++ b/Engine/Plugins/Animation/LiveLinkHub/Source/LiveLinkHub/Private/Config/LiveLinkHubTemplateTokens.h
@@ -2,22 +2,10 @@
 
 #pragma once
 
-#include "NamingTokenContext.h"
 #include "NamingTokens.h"
 
 #include "LiveLinkHubTemplateTokens.generated.h"
 
-UCLASS(NotBlueprintable)
-class ULiveLinkHubNamingTokenContext final: public UNamingTokenContext
-{
-	GENERATED_BODY()
-
-public:
-	/** The name of the current Live Link Hub session. */
-	UPROPERTY(VisibleAnywhere, BlueprintReadOnly, Category = "NamingTokens")
-	FString SessionName;
-};
-
 UCLASS(NotBlueprintable)
 class ULiveLinkHubNamingTokens final: public UNamingTokens
 {
@@ -27,8 +15,13 @@ public:
 	ULiveLinkHubNamingTokens();
 
 protected:
+	// ~Begin UNamingTokens
 	virtual void OnCreateDefaultTokens(TArray<FNamingTokenData>& Tokens) override;
+	virtual void OnPreEvaluate_Implementation(const FNamingTokensEvaluationData& InEvaluationData) override;
+	virtual void OnPostEvaluate_Implementation() override;
+	// ~End UNamingTokens
 	
-	virtual UNamingTokenContext* CreateContext_Implementation() const override;
-	
+private:
+	/** The current session. */
+	FString SessionName;
 };
diff --git a/Engine/Plugins/Animation/LiveLinkHub/Source/LiveLinkHub/Private/Recording/Implementations/LiveLinkUAssetRecorder.cpp b/Engine/Plugins/Animation/LiveLinkHub/Source/LiveLinkHub/Private/Recording/Implementations/LiveLinkUAssetRecorder.cpp
--- a/Engine/Plugins/Animation/LiveLinkHub/Source/LiveLinkHub/Private/Recording/Implementations/LiveLinkUAssetRecorder.cpp
+++ b/Engine/Plugins/Animation/LiveLinkHub/Source/LiveLinkHub/Private/Recording/Implementations/LiveLinkUAssetRecorder.cpp
@@ -158,7 +158,11 @@ bool FLiveLinkUAssetRecorder::GetSavePresetPackageName(FString& OutName)
 	FormatArgs.Add(TEXT("date"), Today.ToString());
 
 	const TObjectPtr<ULiveLinkHubNamingTokens> Tokens = GetDefault<ULiveLinkHubSettings>()->GetNamingTokens();
-	FNamingTokenResultData TemplateData = GEditor->GetEditorSubsystem<UNamingTokensEditorSubsystem>()->EvaluateTokenString(GetDefault<ULiveLinkHubSettings>()->FilenameTemplate, Tokens);
+	FNamingTokenFilterArgs Filter;
+	Filter.NamespacesToInclude.Add(Tokens->GetNamespace());
+	check(GEditor);
+	FNamingTokenResultData TemplateData = GEditor->GetEditorSubsystem<UNamingTokensEditorSubsystem>()->EvaluateTokenString(GetDefault<ULiveLinkHubSettings>()->FilenameTemplate,
+		Filter);
 
 	const FString DefaultName = TemplateData.FileName;
 	const FString DefaultFolder = TemplateData.FolderPath;
diff --git a/Engine/Plugins/Animation/LiveLinkHub/Source/LiveLinkHub/Private/Settings/LiveLinkHubSettings.cpp b/Engine/Plugins/Animation/LiveLinkHub/Source/LiveLinkHub/Private/Settings/LiveLinkHubSettings.cpp
--- a/Engine/Plugins/Animation/LiveLinkHub/Source/LiveLinkHub/Private/Settings/LiveLinkHubSettings.cpp
+++ b/Engine/Plugins/Animation/LiveLinkHub/Source/LiveLinkHub/Private/Settings/LiveLinkHubSettings.cpp
@@ -30,8 +30,12 @@ void ULiveLinkHubSettings::CalculateExampleOutput()
 {
 	if (const TObjectPtr<ULiveLinkHubNamingTokens> Tokens = GetNamingTokens())
 	{
+		FNamingTokenFilterArgs Filter;
+		Filter.NamespacesToInclude.Add(Tokens->GetNamespace());
+		
+		check(GEditor);
 		const FNamingTokenResultData TemplateData = GEditor->GetEditorSubsystem<UNamingTokensEditorSubsystem>()->EvaluateTokenString(
-			GetDefault<ULiveLinkHubSettings>()->FilenameTemplate,Tokens);
+			GetDefault<ULiveLinkHubSettings>()->FilenameTemplate, Filter);
 		FilenameOutput = TemplateData.EvaluatedText.ToString();
 	}
 	else
diff --git a/Engine/Plugins/Developer/NamingTokens/Source/NamingTokens/Private/GlobalNamingTokens.cpp b/Engine/Plugins/Developer/NamingTokens/Source/NamingTokens/Private/GlobalNamingTokens.cpp
--- a/Engine/Plugins/Developer/NamingTokens/Source/NamingTokens/Private/GlobalNamingTokens.cpp
+++ b/Engine/Plugins/Developer/NamingTokens/Source/NamingTokens/Private/GlobalNamingTokens.cpp
@@ -2,8 +2,6 @@
 
 #include "GlobalNamingTokens.h"
 
-#include "NamingTokenContext.h"
-
 #include "Misc/App.h"
 
 #define LOCTEXT_NAMESPACE "GlobalNamingTokens"
@@ -18,141 +16,155 @@ UGlobalNamingTokens::UGlobalNamingTokens()
 void UGlobalNamingTokens::OnCreateDefaultTokens(TArray<FNamingTokenData>& Tokens)
 {
 	Tokens.Add({
-	LOCTEXT("ProjectName", "Project Name"),
-	TEXT("project"),
-		FTokenProcessorDelegateNative::CreateLambda([](const TObjectPtr<UNamingTokenContext>& Context) {
+		TEXT("project"),
+		LOCTEXT("ProjectName", "Project Name"),
+		FNamingTokenData::FTokenProcessorDelegateNative::CreateLambda([]() {
 			return FText::FromString(FApp::GetProjectName());
 		})
 	});
+
 	Tokens.Add({
-		LOCTEXT("Year4Digit", "Year (4 digit)"),
-		TEXT("yyyy"),
-		FTokenProcessorDelegateNative::CreateLambda([](const TObjectPtr<UNamingTokenContext>& Context)
-		{
-			return FText::FromString(FString::Printf(TEXT("%04d"), Context->CurrentDateTime.GetYear()));
-		})
+	    TEXT("yyyy"),
+	    LOCTEXT("Year4Digit", "Year (4 digit)"),
+	    FNamingTokenData::FTokenProcessorDelegateNative::CreateLambda([this]() {
+	        return FText::FromString(FString::Printf(TEXT("%04d"), GetCurrentDateTime().GetYear()));
+	    })
 	});
+
 	Tokens.Add({
-		LOCTEXT("Year2Digit", "Year (2 digit)"),
-		TEXT("yy"),
-		FTokenProcessorDelegateNative::CreateLambda([](const TObjectPtr<UNamingTokenContext>& Context) {
-			return FText::FromString(FString::Printf(TEXT("%02d"), Context->CurrentDateTime.GetYear() % 100));
-		})
+	    TEXT("yy"),
+	    LOCTEXT("Year2Digit", "Year (2 digit)"),
+	    FNamingTokenData::FTokenProcessorDelegateNative::CreateLambda([this]() {
+	        return FText::FromString(FString::Printf(TEXT("%02d"), GetCurrentDateTime().GetYear() % 100));
+	    })
 	});
+
 	Tokens.Add({
-		LOCTEXT("MonthPascalCase", "3-character Month (Pascal Case)"),
-		TEXT("Mmm"),
-		FTokenProcessorDelegateNative::CreateLambda([](const TObjectPtr<UNamingTokenContext>& Context) {
-			const FString MonthName = Context->CurrentDateTime.ToFormattedString(TEXT("%b"));
-			return FText::FromString(MonthName.Left(1).ToUpper() + MonthName.Mid(1).ToLower());
-		})
+	    TEXT("Mmm"),
+	    LOCTEXT("MonthPascalCase", "3-character Month (Pascal Case)"),
+	    FNamingTokenData::FTokenProcessorDelegateNative::CreateLambda([this]() {
+	        const FString MonthName = GetCurrentDateTime().ToFormattedString(TEXT("%b"));
+	        return FText::FromString(MonthName.Left(1).ToUpper() + MonthName.Mid(1).ToLower());
+	    })
 	});
+
 	Tokens.Add({
-		LOCTEXT("MonthUpperCase", "3-character Month (UPPERCASE)"),
-		TEXT("MMM"),
-		FTokenProcessorDelegateNative::CreateLambda([](const TObjectPtr<UNamingTokenContext>& Context) {
-			const FString MonthName = Context->CurrentDateTime.ToFormattedString(TEXT("%b"));
-			return FText::FromString(MonthName.ToUpper());
-		})
+	    TEXT("MMM"),
+	    LOCTEXT("MonthUpperCase", "3-character Month (UPPERCASE)"),
+	    FNamingTokenData::FTokenProcessorDelegateNative::CreateLambda([this]() {
+	        const FString MonthName = GetCurrentDateTime().ToFormattedString(TEXT("%b"));
+	        return FText::FromString(MonthName.ToUpper());
+	    })
 	});
+
 	Tokens.Add({
-		LOCTEXT("MonthLowerCase", "3-character Month (lowercase)"),
-		TEXT("mmm"),
-		FTokenProcessorDelegateNative::CreateLambda([](const TObjectPtr<UNamingTokenContext>& Context) {
-			const FString MonthName = Context->CurrentDateTime.ToFormattedString(TEXT("%b"));
-			return FText::FromString(MonthName.ToLower());
-		})
+	    TEXT("mmm"),
+	    LOCTEXT("MonthLowerCase", "3-character Month (lowercase)"),
+	    FNamingTokenData::FTokenProcessorDelegateNative::CreateLambda([this]() {
+	        const FString MonthName = GetCurrentDateTime().ToFormattedString(TEXT("%b"));
+	        return FText::FromString(MonthName.ToLower());
+	    })
 	});
+
 	Tokens.Add({
-		LOCTEXT("Month2Digit", "Month (2 digit)"),
-		TEXT("mm"),
-		FTokenProcessorDelegateNative::CreateLambda([](const TObjectPtr<UNamingTokenContext>& Context) {
-			return FText::FromString(FString::Printf(TEXT("%02d"), Context->CurrentDateTime.GetMonth()));
-		})
+	    TEXT("mm"),
+	    LOCTEXT("Month2Digit", "Month (2 digit)"),
+	    FNamingTokenData::FTokenProcessorDelegateNative::CreateLambda([this]() {
+	        return FText::FromString(FString::Printf(TEXT("%02d"), GetCurrentDateTime().GetMonth()));
+	    })
 	});
+
 	Tokens.Add({
-		LOCTEXT("DayPascalCase", "3-character Day (Pascal Case)"),
-		TEXT("Ddd"),
-		FTokenProcessorDelegateNative::CreateLambda([](const TObjectPtr<UNamingTokenContext>& Context) {
-			const FString DayName = Context->CurrentDateTime.ToFormattedString(TEXT("%a"));
-			return FText::FromString(DayName.Left(1).ToUpper() + DayName.Mid(1).ToLower());
-		})
+	    TEXT("Ddd"),
+	    LOCTEXT("DayPascalCase", "3-character Day (Pascal Case)"),
+	    FNamingTokenData::FTokenProcessorDelegateNative::CreateLambda([this]() {
+	        const FString DayName = GetCurrentDateTime().ToFormattedString(TEXT("%a"));
+	        return FText::FromString(DayName.Left(1).ToUpper() + DayName.Mid(1).ToLower());
+	    })
 	});
 
 	Tokens.Add({
-		LOCTEXT("DayUpperCase", "3-character Day (UPPERCASE)"),
-		TEXT("DDD"),
-		FTokenProcessorDelegateNative::CreateLambda([](const TObjectPtr<UNamingTokenContext>& Context) {
-			const FString DayName = Context->CurrentDateTime.ToFormattedString(TEXT("%a"));
-			return FText::FromString(DayName.ToUpper());
-		})
+	    TEXT("DDD"),
+	    LOCTEXT("DayUpperCase", "3-character Day (UPPERCASE)"),
+	    FNamingTokenData::FTokenProcessorDelegateNative::CreateLambda([this]() {
+	        const FString DayName = GetCurrentDateTime().ToFormattedString(TEXT("%a"));
+	        return FText::FromString(DayName.ToUpper());
+	    })
 	});
 
 	Tokens.Add({
-		LOCTEXT("DayLowerCase", "3-character Day (lowercase)"),
-		TEXT("ddd"),
-		FTokenProcessorDelegateNative::CreateLambda([](const TObjectPtr<UNamingTokenContext>& Context) {
-			const FString DayName = Context->CurrentDateTime.ToFormattedString(TEXT("%a"));
-			return FText::FromString(DayName.ToLower());
-		})
+	    TEXT("ddd"),
+	    LOCTEXT("DayLowerCase", "3-character Day (lowercase)"),
+	    FNamingTokenData::FTokenProcessorDelegateNative::CreateLambda([this]() {
+	        const FString DayName = GetCurrentDateTime().ToFormattedString(TEXT("%a"));
+	        return FText::FromString(DayName.ToLower());
+	    })
 	});
+
 	Tokens.Add({
-		LOCTEXT("Day2Digit", "Day (2 digit)"),
-		TEXT("dd"),
-		FTokenProcessorDelegateNative::CreateLambda([](const TObjectPtr<UNamingTokenContext>& Context) {
-			return FText::FromString(FString::Printf(TEXT("%02d"), Context->CurrentDateTime.GetDay()));
-		})
+	    TEXT("dd"),
+	    LOCTEXT("Day2Digit", "Day (2 digit)"),
+	    FNamingTokenData::FTokenProcessorDelegateNative::CreateLambda([this]() {
+	        return FText::FromString(FString::Printf(TEXT("%02d"), GetCurrentDateTime().GetDay()));
+	    })
 	});
+
 	Tokens.Add({
-	LOCTEXT("AMPMLowerCase", "am or pm (lowercase)"),
-	TEXT("ampm"),
-		FTokenProcessorDelegateNative::CreateLambda([](const TObjectPtr<UNamingTokenContext>& Context) {
-			const FString AMPM = Context->CurrentDateTime.ToFormattedString(TEXT("%P"));
-			return FText::FromString(AMPM);
-		})
+	    TEXT("ampm"),
+	    LOCTEXT("AMPMLowerCase", "am or pm (lowercase)"),
+	    FNamingTokenData::FTokenProcessorDelegateNative::CreateLambda([this]() {
+	        const FString AMPM = GetCurrentDateTime().ToFormattedString(TEXT("%P"));
+	        return FText::FromString(AMPM);
+	    })
 	});
+
 	Tokens.Add({
-	LOCTEXT("AMPMUpperCase", "AM or PM (UPPERCASE)"),
-	TEXT("AMPM"),
-		FTokenProcessorDelegateNative::CreateLambda([](const TObjectPtr<UNamingTokenContext>& Context) {
-			const FString AMPM = Context->CurrentDateTime.ToFormattedString(TEXT("%p"));
-			return FText::FromString(AMPM);
-		})
+	    TEXT("AMPM"),
+	    LOCTEXT("AMPMUpperCase", "AM or PM (UPPERCASE)"),
+	    FNamingTokenData::FTokenProcessorDelegateNative::CreateLambda([this]() {
+	        const FString AMPM = GetCurrentDateTime().ToFormattedString(TEXT("%p"));
+	        return FText::FromString(AMPM);
+	    })
 	});
+
 	Tokens.Add({
-	LOCTEXT("Hour12", "Hour (12)"),
-	TEXT("12h"),
-		FTokenProcessorDelegateNative::CreateLambda([](const TObjectPtr<UNamingTokenContext>& Context) {
-			return FText::FromString(FString::Printf(TEXT("%02d"), Context->CurrentDateTime.GetHour12()));
-		})
+	    TEXT("12h"),
+	    LOCTEXT("Hour12", "Hour (12)"),
+	    FNamingTokenData::FTokenProcessorDelegateNative::CreateLambda([this]() {
+	        return FText::FromString(FString::Printf(TEXT("%02d"), GetCurrentDateTime().GetHour12()));
+	    })
 	});
+
 	Tokens.Add({
-	LOCTEXT("Hour24", "Hour (24)"),
-	TEXT("24h"),
-		FTokenProcessorDelegateNative::CreateLambda([](const TObjectPtr<UNamingTokenContext>& Context) {
-			return FText::FromString(FString::Printf(TEXT("%02d"), Context->CurrentDateTime.GetHour()));
-		})
+	    TEXT("24h"),
+	    LOCTEXT("Hour24", "Hour (24)"),
+	    FNamingTokenData::FTokenProcessorDelegateNative::CreateLambda([this]() {
+	        return FText::FromString(FString::Printf(TEXT("%02d"), GetCurrentDateTime().GetHour()));
+	    })
 	});
+
 	Tokens.Add({
-	LOCTEXT("Minute", "Minute"),
-	TEXT("min"),
-		FTokenProcessorDelegateNative::CreateLambda([](const TObjectPtr<UNamingTokenContext>& Context) {
-			return FText::FromString(FString::Printf(TEXT("%02d"), Context->CurrentDateTime.GetMinute()));
-		})
+	    TEXT("min"),
+	    LOCTEXT("Minute", "Minute"),
+	    FNamingTokenData::FTokenProcessorDelegateNative::CreateLambda([this]() {
+	        return FText::FromString(FString::Printf(TEXT("%02d"), GetCurrentDateTime().GetMinute()));
+	    })
 	});
+
 	Tokens.Add({
-	LOCTEXT("Second", "Second"),
-	TEXT("sec"),
-		FTokenProcessorDelegateNative::CreateLambda([](const TObjectPtr<UNamingTokenContext>& Context) {
-			return FText::FromString(FString::Printf(TEXT("%02d"), Context->CurrentDateTime.GetSecond()));
-		})
+	    TEXT("sec"),
+	    LOCTEXT("Second", "Second"),
+	    FNamingTokenData::FTokenProcessorDelegateNative::CreateLambda([this]() {
+	        return FText::FromString(FString::Printf(TEXT("%02d"), GetCurrentDateTime().GetSecond()));
+	    })
 	});
+
 	Tokens.Add({
-	LOCTEXT("Millisecond", "Millisecond"),
-	TEXT("ms"),
-		FTokenProcessorDelegateNative::CreateLambda([](const TObjectPtr<UNamingTokenContext>& Context) {
-			return FText::FromString(FString::Printf(TEXT("%02d"), Context->CurrentDateTime.GetMillisecond()));
-		})
+	    TEXT("ms"),
+	    LOCTEXT("Millisecond", "Millisecond"),
+	    FNamingTokenData::FTokenProcessorDelegateNative::CreateLambda([this]() {
+	        return FText::FromString(FString::Printf(TEXT("%02d"), GetCurrentDateTime().GetMillisecond()));
+	    })
 	});
 }
 
diff --git a/Engine/Plugins/Developer/NamingTokens/Source/NamingTokens/Private/NamingTokens.cpp b/Engine/Plugins/Developer/NamingTokens/Source/NamingTokens/Private/NamingTokens.cpp
--- a/Engine/Plugins/Developer/NamingTokens/Source/NamingTokens/Private/NamingTokens.cpp
+++ b/Engine/Plugins/Developer/NamingTokens/Source/NamingTokens/Private/NamingTokens.cpp
@@ -2,10 +2,11 @@
 
 #include "NamingTokens.h"
 
-#include "NamingTokenContext.h"
+#include "NamingTokensEvaluationData.h"
 #include "NamingTokensLog.h"
 #include "Utils/NamingTokenUtils.h"
 
+#include "Engine/Engine.h"
 #include "Internationalization/Regex.h"
 #include "UObject/Package.h"
 
@@ -13,7 +14,6 @@
 
 UNamingTokens::UNamingTokens()
 {
-	NamingTokenContextClass = UNamingTokenContext::StaticClass();
 }
 
 void UNamingTokens::PostInitProperties()
@@ -40,7 +40,7 @@ void UNamingTokens::PostEditChangeProperty(FPropertyChangedEvent& PropertyChange
 
 void UNamingTokens::EvaluateTestToken()
 {
-	const FNamingTokenResultData ResultData = EvaluateTokenString(TestTokenInput.ToString());
+	const FNamingTokenResultData ResultData = EvaluateTokenText(TestTokenInput);
 	TestTokenResult = ResultData.EvaluatedText;
 }
 #endif
@@ -49,30 +49,34 @@ void UNamingTokens::OnCreateDefaultTokens(TArray<FNamingTokenData>& Tokens)
 {
 }
 
-FNamingTokenResultData UNamingTokens::EvaluateTokenString(const FString& InTokenString)
+FNamingTokenResultData UNamingTokens::EvaluateTokenText(const FText& InTokenText)
 {
-	FNamingTokenResultData Result;
-	Result.OriginalString = InTokenString;
-
-	const FString RawTokenString = PreProcessTokenString(InTokenString);
-	FText FormattedText = FText::FromString(RawTokenString);
+	return EvaluateTokenText(InTokenText, FNamingTokensEvaluationData());
+}
 
-	const TObjectPtr<UNamingTokenContext> Context = CreateContext();
+FNamingTokenResultData UNamingTokens::EvaluateTokenText(const FText& InTokenText, const FNamingTokensEvaluationData& InEvaluationData)
+{
+	PreEvaluate(InEvaluationData);
 	
-	for (const FNamingTokenData& Token : GetDefaultTokens())
-	{
-		TryProcessToken(Token,Context, FormattedText);
-	}
+	FNamingTokenResultData Result;
+	Result.OriginalText = InTokenText;
 
-	for (const FNamingTokenData& Token : CustomTokens)
+	FText FormattedText = InTokenText;
 	{
-		TryProcessToken(Token,Context, FormattedText);
-	}
+		for (const FNamingTokenData& Token : GetDefaultTokens())
+		{
+			TryProcessToken(Token, FormattedText);
+		}
 
+		for (const FNamingTokenData& Token : CustomTokens)
+		{
+			TryProcessToken(Token, FormattedText);
+		}
+	}
 	Result.EvaluatedText = FormattedText;
-	FString FormattedString = FormattedText.ToString();
-
+	
 	// Split folder path and file name.
+	FString FormattedString = FormattedText.ToString();
 	int32 LastSlashIndex;
 	if (FormattedString.FindLastChar('/', LastSlashIndex))
 	{
@@ -88,15 +92,12 @@ FNamingTokenResultData UNamingTokens::EvaluateTokenString(const FString& InToken
 		Result.FolderPath = TEXT("");
 		Result.FileName = FormattedString;
 	}
+
+	PostEvaluate();
 	
 	return Result;
 }
 
-FString UNamingTokens::CreateFormattedToken(const FNamingTokenData& InToken) const
-{
-	return FString::Printf(TEXT("{%s}"), *InToken.TokenKey);
-}
-
 FString UNamingTokens::GetFormattedTokensStringForDisplay() const
 {
 	FString FormattedTokensString;
@@ -122,52 +123,31 @@ const TArray<FNamingTokenData>& UNamingTokens::GetDefaultTokens() const
 	return DefaultTokens;
 }
 
-UNamingTokenContext* UNamingTokens::CreateContext_Implementation() const
+FDateTime UNamingTokens::GetCurrentDateTime_Implementation() const
 {
-	const UClass* ContextClass = NamingTokenContextClass.LoadSynchronous();
-	if (!ContextClass)
-	{
-		ContextClass = UNamingTokenContext::StaticClass();
-	}
-	
-	TObjectPtr<UNamingTokenContext> Context = NewObject<UNamingTokenContext>(GetTransientPackage(), ContextClass);
-	Context->CurrentDateTime = FDateTime::Now();
-
-	return Context;
+	return CurrentEvaluationData.CurrentDateTime;
 }
 
-FString UNamingTokens::PreProcessTokenString(const FString& InTokenString) const
+void UNamingTokens::PreEvaluate(const FNamingTokensEvaluationData& InEvaluationData)
 {
-	return InTokenString;
+	CurrentEvaluationData = InEvaluationData;
+	OnPreEvaluate(InEvaluationData);
 }
 
-FString UNamingTokens::RemoveNamespaceFromTokenString(const FString& InTokenString) const
+void UNamingTokens::PostEvaluate()
 {
-	// Build the regex pattern to match {namespace.token}
-	if (Namespace.IsEmpty())
-	{
-		return InTokenString;
-	}
-	const FString PatternString = FString::Printf(TEXT(R"(\{\s*%s\.([a-zA-Z0-9_]+)\s*\})"), *Namespace);
-	const FRegexPattern Pattern(PatternString);
-
-	FRegexMatcher Matcher(Pattern, InTokenString);
-
-	FString ProcessedText = InTokenString;
-	while (Matcher.FindNext())
-	{
-		// Get the token part of {namespace.token}
-		FString Token = Matcher.GetCaptureGroup(1);
+	OnPostEvaluate();
+}
 
-		// Replace {namespace.token} with {token}
-		const FString FullMatch = Matcher.GetCaptureGroup(0);
-		ProcessedText = ProcessedText.Replace(*FullMatch, *FString::Printf(TEXT("{%s}"), *Token));
-	}
+void UNamingTokens::OnPreEvaluate_Implementation(const FNamingTokensEvaluationData& InEvaluationData)
+{
+}
 
-	return ProcessedText;
+void UNamingTokens::OnPostEvaluate_Implementation()
+{
 }
 
-bool UNamingTokens::TryProcessToken(const FNamingTokenData& InToken, UNamingTokenContext* InContext, FText& InOutFormattedText)
+bool UNamingTokens::TryProcessToken(const FNamingTokenData& InToken, FText& InOutFormattedText)
 {
 	const FString FormattedToken = *CreateFormattedToken(InToken);
 
@@ -177,16 +157,14 @@ bool UNamingTokens::TryProcessToken(const FNamingTokenData& InToken, UNamingToke
 	{
 		struct FFunctionParams
 		{
-			UNamingTokenContext* Input;
 			FText ReturnValue;
 		} Params;
-		Params.Input = InContext;
 		this->ProcessEvent(BlueprintFunction, &Params);
 		ProcessedToken = Params.ReturnValue;
 	}
 	else if (InToken.TokenProcessorNative.IsBound())
 	{
-		ProcessedToken = InToken.TokenProcessorNative.Execute(InContext);
+		ProcessedToken = InToken.TokenProcessorNative.Execute();
 	}
 	else
 	{
@@ -202,6 +180,11 @@ bool UNamingTokens::TryProcessToken(const FNamingTokenData& InToken, UNamingToke
 	return true;
 }
 
+FString UNamingTokens::CreateFormattedToken(const FNamingTokenData& InToken) const
+{
+	return FString::Printf(TEXT("{%s}"), *InToken.TokenKey);
+}
+
 UFunction* UNamingTokens::FindBlueprintFunctionForToken(const FNamingTokenData& InTokenData) const
 {
 	const FName FunctionName = InTokenData.FunctionName;
diff --git a/Engine/Plugins/Developer/NamingTokens/Source/NamingTokens/Private/NamingTokensEditorSubsystem.cpp b/Engine/Plugins/Developer/NamingTokens/Source/NamingTokens/Private/NamingTokensEditorSubsystem.cpp
--- a/Engine/Plugins/Developer/NamingTokens/Source/NamingTokens/Private/NamingTokensEditorSubsystem.cpp
+++ b/Engine/Plugins/Developer/NamingTokens/Source/NamingTokens/Private/NamingTokensEditorSubsystem.cpp
@@ -3,6 +3,7 @@
 #include "NamingTokensEditorSubsystem.h"
 
 #include "NamingTokens.h"
+#include "NamingTokensEvaluationData.h"
 #include "NamingTokensLog.h"
 
 #include "AssetRegistry/AssetData.h"
@@ -154,33 +155,34 @@ TArray<UNamingTokens*> UNamingTokensEditorSubsystem::GetMultipleNamingTokens(con
 	return Result;
 }
 
-FNamingTokenResultData UNamingTokensEditorSubsystem::EvaluateTokenString(const FString& InTokenString, FNamingTokenFilterArgs InFilter)
+FNamingTokenResultData UNamingTokensEditorSubsystem::EvaluateTokenText(const FText& InTokenText, FNamingTokenFilterArgs InFilter)
 {
 	FNamingTokenResultData Result;
 
-	FText ProcessedTokenText = FText::FromString(InTokenString);
+	FText ProcessedTokenText = InTokenText;
 	Result.EvaluatedText = ProcessedTokenText; // Set so we always have something in case no tokens were evaluated.
 
-	TSet<FString> Namespaces = GetNamingTokenNamespacesFromString(InTokenString, InFilter);
+	TSet<FString> Namespaces = GetNamingTokenNamespacesFromString(InTokenText.ToString(), InFilter);
 	Namespaces.Append(InFilter.NamespacesToInclude); // Tokens could have been written without the namespace but should still be scoped to the filter.
-	
+
+	const FNamingTokensEvaluationData EvaluationData;
 	for (const FString& Namespace : Namespaces)
 	{
 		if (UNamingTokens* NamingTokens = GetNamingTokens(Namespace))
 		{
-			Result = NamingTokens->EvaluateTokenString(ProcessedTokenText.ToString());
+			Result = NamingTokens->EvaluateTokenText(ProcessedTokenText, EvaluationData);
 			ProcessedTokenText = Result.EvaluatedText;
 		}
 	}
 
-	Result.OriginalString = InTokenString;
+	Result.OriginalText = InTokenText;
 	
 	return Result;
 }
 
-FNamingTokenResultData UNamingTokensEditorSubsystem::EvaluateTokenString(const FString& InTokenString, const UNamingTokens* InNamingTokens, bool bIncludeGlobal)
+FNamingTokenResultData UNamingTokensEditorSubsystem::EvaluateTokenString(const FString& InTokenString, FNamingTokenFilterArgs InFilter)
 {
-	return EvaluateTokenString(InTokenString, { { InNamingTokens->GetNamespace() }, bIncludeGlobal });
+	return EvaluateTokenText(FText::FromString(InTokenString), InFilter);
 }
 
 void UNamingTokensEditorSubsystem::RegisterGlobalNamespace(const FString& InNamespace)
@@ -219,9 +221,7 @@ TSet<FString> UNamingTokensEditorSubsystem::GetNamingTokenNamespacesFromString(c
 	
 	while (Matcher.FindNext())
 	{
-		// Get the first capture group (namespace)
 		const FString Namespace = Matcher.GetCaptureGroup(1);
-
 		if (InFilter.NamespacesToInclude.IsEmpty() || InFilter.NamespacesToInclude.Contains(Namespace))
 		{
 			Namespaces.Add(Namespace);
diff --git a/Engine/Plugins/Developer/NamingTokens/Source/NamingTokens/Private/NamingTokensEvaluationData.cpp b/Engine/Plugins/Developer/NamingTokens/Source/NamingTokens/Private/NamingTokensEvaluationData.cpp
new file mode 100644
--- /dev/null
+++ b/Engine/Plugins/Developer/NamingTokens/Source/NamingTokens/Private/NamingTokensEvaluationData.cpp
@@ -0,0 +1,8 @@
+// Copyright Epic Games, Inc. All Rights Reserved.
+
+#include "NamingTokensEvaluationData.h"
+
+FNamingTokensEvaluationData::FNamingTokensEvaluationData()
+{
+	CurrentDateTime = FDateTime::Now();
+}
diff --git a/Engine/Plugins/Developer/NamingTokens/Source/NamingTokens/Private/NamingTokensModule.cpp b/Engine/Plugins/Developer/NamingTokens/Source/NamingTokens/Private/NamingTokensModule.cpp
--- a/Engine/Plugins/Developer/NamingTokens/Source/NamingTokens/Private/NamingTokensModule.cpp
+++ b/Engine/Plugins/Developer/NamingTokens/Source/NamingTokens/Private/NamingTokensModule.cpp
@@ -28,7 +28,10 @@ void FNamingTokensModule::ShutdownModule()
 
 	if (GEditor)
 	{
-		GEditor->GetEditorSubsystem<UNamingTokensEditorSubsystem>()->UnregisterGlobalNamespace(UGlobalNamingTokens::GetGlobalNamespace());
+		if (UNamingTokensEditorSubsystem* NamingTokensSubsystem = GEditor->GetEditorSubsystem<UNamingTokensEditorSubsystem>())
+		{
+			NamingTokensSubsystem->UnregisterGlobalNamespace(UGlobalNamingTokens::GetGlobalNamespace());
+		}
 	}
 	
 	FCoreDelegates::OnPostEngineInit.RemoveAll(this);
diff --git a/Engine/Plugins/Developer/NamingTokens/Source/NamingTokens/Public/GlobalNamingTokens.h b/Engine/Plugins/Developer/NamingTokens/Source/NamingTokens/Public/GlobalNamingTokens.h
--- a/Engine/Plugins/Developer/NamingTokens/Source/NamingTokens/Public/GlobalNamingTokens.h
+++ b/Engine/Plugins/Developer/NamingTokens/Source/NamingTokens/Public/GlobalNamingTokens.h
@@ -6,8 +6,6 @@
 
 #include "GlobalNamingTokens.generated.h"
 
-class UNamingTokenContext;
-
 /**
  * Default global tokens accessed project wide.
  */
diff --git a/Engine/Plugins/Developer/NamingTokens/Source/NamingTokens/Public/NamingTokenContext.h b/Engine/Plugins/Developer/NamingTokens/Source/NamingTokens/Public/NamingTokenContext.h
deleted file mode 100644
diff --git a/Engine/Plugins/Developer/NamingTokens/Source/NamingTokens/Public/NamingTokenData.h b/Engine/Plugins/Developer/NamingTokens/Source/NamingTokens/Public/NamingTokenData.h
--- a/Engine/Plugins/Developer/NamingTokens/Source/NamingTokens/Public/NamingTokenData.h
+++ b/Engine/Plugins/Developer/NamingTokens/Source/NamingTokens/Public/NamingTokenData.h
@@ -7,18 +7,16 @@
 
 #include "NamingTokenData.generated.h"
 
-class UNamingTokenContext;
-
 /** Evaluated results from a template file string. */
 USTRUCT(BlueprintType)
 struct NAMINGTOKENS_API FNamingTokenResultData
 {
 	GENERATED_BODY()
 	
-	/** Original string without any modifications. */
+	/** Original text without any modifications. */
 	UPROPERTY(BlueprintReadOnly, Category = "NamingTokens")
-	FString OriginalString;
-	/** The full string with evaluated tokens. */
+	FText OriginalText;
+	/** The full text with evaluated tokens. */
 	UPROPERTY(BlueprintReadOnly, Category = "NamingTokens")
 	FText EvaluatedText;
 	/** The evaluated folder(s) without the first forward slash or filename. */
@@ -29,8 +27,6 @@ struct NAMINGTOKENS_API FNamingTokenResultData
 	FString FileName;
 };
 
-DECLARE_DELEGATE_RetVal_OneParam(FText, FTokenProcessorDelegateNative, const TObjectPtr<UNamingTokenContext>&);
-
 USTRUCT(BlueprintType)
 struct NAMINGTOKENS_API FNamingTokenData
 {
@@ -38,14 +34,22 @@ struct NAMINGTOKENS_API FNamingTokenData
 
 public:
 	FNamingTokenData() = default;
-	
-	FNamingTokenData(const FText& InTokenDisplayName, const FString& InTokenKey, const FTokenProcessorDelegateNative& InTokenProcessor)
+
+	DECLARE_DELEGATE_RetVal(FText, FTokenProcessorDelegateNative);
+	FNamingTokenData(const FString& InTokenKey, const FText& InTokenDisplayName, const FTokenProcessorDelegateNative& InTokenProcessor)
 	{
 		DisplayName = InTokenDisplayName;
 		TokenKey = InTokenKey;
 		TokenProcessorNative = InTokenProcessor;
 	}
 
+	FNamingTokenData(const FString& InTokenKey, const FText& InTokenDisplayName,
+		const FText& InTokenDescription, const FTokenProcessorDelegateNative& InTokenProcessor)
+		: FNamingTokenData(InTokenKey, InTokenDisplayName, InTokenProcessor)
+	{
+		Description = InTokenDescription;
+	}
+
 	/** The key of the token to use. This is what the text must match in order to be evaluated. Brackets are automatically added and do not need to be included. */
 	UPROPERTY(BlueprintReadWrite, EditAnywhere, Category = "NamingTokens")
 	FString TokenKey;
diff --git a/Engine/Plugins/Developer/NamingTokens/Source/NamingTokens/Public/NamingTokens.h b/Engine/Plugins/Developer/NamingTokens/Source/NamingTokens/Public/NamingTokens.h
--- a/Engine/Plugins/Developer/NamingTokens/Source/NamingTokens/Public/NamingTokens.h
+++ b/Engine/Plugins/Developer/NamingTokens/Source/NamingTokens/Public/NamingTokens.h
@@ -3,13 +3,11 @@
 #pragma once
 
 #include "CoreMinimal.h"
-#include "UObject/ObjectMacros.h"
 #include "NamingTokenData.h"
+#include "NamingTokensEvaluationData.h"
 
 #include "NamingTokens.generated.h"
 
-class UNamingTokenContext;
-
 /**
  * Subclass to define naming tokens to use for a specific tool or project.
  */
@@ -37,12 +35,12 @@ protected:
 	virtual void OnCreateDefaultTokens(TArray<FNamingTokenData>& Tokens);
 
 public:
-	/** Evaluate a token string. */
-	virtual FNamingTokenResultData EvaluateTokenString(const FString& InTokenString);
-
-	/** Formats a token key into {key}. */
-	FString CreateFormattedToken(const FNamingTokenData& InToken) const;
+	/** Evaluate token text. Creates EvaluationData. */
+	FNamingTokenResultData EvaluateTokenText(const FText& InTokenText);
 
+	/** Evaluate token text for a given EvaluationData. */
+	FNamingTokenResultData EvaluateTokenText(const FText& InTokenText, const FNamingTokensEvaluationData& InEvaluationData);
+	
 	/** Creates a friendly display string of all tokens. */
 	FString GetFormattedTokensStringForDisplay() const;
 
@@ -54,33 +52,61 @@ public:
 	/** The name of the namespace property. */
 	static FName GetNamespacePropertyName() { return GET_MEMBER_NAME_CHECKED(UNamingTokens, Namespace); }
 
+protected:
+	/** Retrieve the current datetime. By default, this uses shared data so results are consistent across runs. */
+	UFUNCTION(BlueprintCallable, BlueprintNativeEvent, Category = "NamingTokens")
+	FDateTime GetCurrentDateTime() const;
+	virtual FDateTime GetCurrentDateTime_Implementation() const;
+	
 #if WITH_EDITOR
 public:
 	static FName GetProcessTokenTemplateFunctionName() { return GET_FUNCTION_NAME_CHECKED(UNamingTokens, ProcessTokenTemplateFunction); } 
 private:
 	/** Template function for us to dynamically create subclass graphs from matching this signature. */
 	UFUNCTION()
-	FText ProcessTokenTemplateFunction(UNamingTokenContext* InContext) { return FText::GetEmpty(); }
+	FText ProcessTokenTemplateFunction() { return FText::GetEmpty(); }
 #endif
-	
+
 protected:
-	/** Create the context to use for token evaluation. */
+	/**
+	 * Called prior to token evaluation.
+	 * 
+	 * @param InEvaluationData Shared information across namespaces.
+	 */
+	void PreEvaluate(const FNamingTokensEvaluationData& InEvaluationData);
+	
+	/**
+	 * Called after all tokens have evaluated.
+	 */
+	void PostEvaluate();
+	
+	/**
+	 * Called prior to evaluation. Allows consistent data to be setup for each token evaluation.
+	 * 
+	 * This is important if the data is temporally sensitive and could change between evaluating
+	 * multiple tokens in a string, such as a high resolution timer.
+	 * 
+	 * @param InEvaluationData Contains generic shared data for this evaluation.
+	 */
 	UFUNCTION(BlueprintNativeEvent, Category = "NamingTokens")
-	UNamingTokenContext* CreateContext() const;
-	virtual UNamingTokenContext* CreateContext_Implementation() const;
-
-	/** Run pre-processing logic on the token string prior to evaluation. */
-	virtual FString PreProcessTokenString(const FString& InTokenString) const;
-
-	/** Remove any uses of our namespace from the string. */
-	FString RemoveNamespaceFromTokenString(const FString& InTokenString) const;
+	void OnPreEvaluate(const FNamingTokensEvaluationData& InEvaluationData);
+	virtual void OnPreEvaluate_Implementation(const FNamingTokensEvaluationData& InEvaluationData);
 	
+	/** Called after evaluation. Used so any cleanup can occur. */
+	UFUNCTION(BlueprintNativeEvent, Category = "NamingTokens")
+	void OnPostEvaluate();
+	virtual void OnPostEvaluate_Implementation();
+
+private:
 	/** Process a token if is defined and return evaluated text. */
-	bool TryProcessToken(const FNamingTokenData& InToken, UNamingTokenContext* InContext, FText& InOutFormattedText);
+	bool TryProcessToken(const FNamingTokenData& InToken, FText& InOutFormattedText);
+
+	/** Formats a token key into {key}. */
+	FString CreateFormattedToken(const FNamingTokenData& InToken) const;
 	
 	/** Return a blueprint function for a token processor if it exists. */
 	UFunction* FindBlueprintFunctionForToken(const FNamingTokenData& InTokenData) const;
-	
+
 private:
 	/** The default tokens defined by this class. */
 	TArray<FNamingTokenData> DefaultTokens;
@@ -89,16 +115,17 @@ protected:
 	/** User defined tokens. */
 	UPROPERTY(EditDefaultsOnly, Category = "NamingTokens", meta = (TitleProperty="TokenKey"))
 	TArray<FNamingTokenData> CustomTokens;
+
+	/** Cached shared data for this evaluation. */
+	UPROPERTY(BlueprintReadOnly, Category = "NamingTokens")
+	FNamingTokensEvaluationData CurrentEvaluationData;
 	
 	/** The namespace to identify this token. */
 	UPROPERTY(EditDefaultsOnly, AssetRegistrySearchable, Category = "NamingTokens")
 	FString Namespace;
 	
-	/** The class to use for the naming token context. */
-	UPROPERTY(EditDefaultsOnly, Category = "NamingTokens", NoClear)
-	TSoftClassPtr<UNamingTokenContext> NamingTokenContextClass;
-	
 #if WITH_EDITORONLY_DATA
+private:
 	/** Enter a sample string using your tokens to output an evaluated result to TestTokenResult. */
 	UPROPERTY(EditDefaultsOnly, Transient, Category = "NamingTokens")
 	FText TestTokenInput;
@@ -108,5 +135,3 @@ protected:
 	FText TestTokenResult;
 #endif
 };
-
-
diff --git a/Engine/Plugins/Developer/NamingTokens/Source/NamingTokens/Public/NamingTokensEditorSubsystem.h b/Engine/Plugins/Developer/NamingTokens/Source/NamingTokens/Public/NamingTokensEditorSubsystem.h
--- a/Engine/Plugins/Developer/NamingTokens/Source/NamingTokens/Public/NamingTokensEditorSubsystem.h
+++ b/Engine/Plugins/Developer/NamingTokens/Source/NamingTokens/Public/NamingTokensEditorSubsystem.h
@@ -15,7 +15,10 @@ struct FNamingTokenFilterArgs
 {
 	GENERATED_BODY()
 	
-	/** Namespaces to include in the search. */
+	/**
+	 * Namespaces to always be included during evaluation. Namespaces added here won't require the 'namespace.' string prefixed to tokens.
+	 * This does not filter out any namespaces.
+	 */
 	UPROPERTY(EditAnywhere, BlueprintReadWrite, Category=NamingTokens)
 	TArray<FString> NamespacesToInclude;
 	/** Include global namespaces. */
@@ -52,13 +55,14 @@ public:
 	UFUNCTION(BlueprintCallable, Category = "NamingTokens")
 	TArray<UNamingTokens*> GetMultipleNamingTokens(const TArray<FString>& InNamespaces) const;
 
+	/** Parse and evaluate token text. */
+	UFUNCTION(BlueprintCallable, Category = "NamingTokens")
+	FNamingTokenResultData EvaluateTokenText(const FText& InTokenText, FNamingTokenFilterArgs InFilter);
+
 	/** Parse and evaluate a token string. */
 	UFUNCTION(BlueprintCallable, Category = "NamingTokens")
 	FNamingTokenResultData EvaluateTokenString(const FString& InTokenString, FNamingTokenFilterArgs InFilter);
 
-	/** Evaluate a token string for a specific token instance. */
-	FNamingTokenResultData EvaluateTokenString(const FString& InTokenString, const UNamingTokens* InNamingTokens, bool bIncludeGlobal = true);
-
 	/**
 	 * Register tokens as a global namespace. This prevents the need to include the namespace in a token string.
 	 * @param InNamespace The namespace of the tokens to register.
diff --git a/Engine/Plugins/Developer/NamingTokens/Source/NamingTokens/Public/NamingTokensEvaluationData.h b/Engine/Plugins/Developer/NamingTokens/Source/NamingTokens/Public/NamingTokensEvaluationData.h
new file mode 100644
--- /dev/null
+++ b/Engine/Plugins/Developer/NamingTokens/Source/NamingTokens/Public/NamingTokensEvaluationData.h
@@ -0,0 +1,24 @@
+// Copyright Epic Games, Inc. All Rights Reserved.
+
+#pragma once
+
+#include "CoreMinimal.h"
+#include "Misc/DateTime.h"
+
+#include "NamingTokensEvaluationData.generated.h"
+
+/**
+ * Shared data across token evaluation.
+ */
+USTRUCT(BlueprintType)
+struct FNamingTokensEvaluationData
+{
+	GENERATED_BODY()
+
+public:
+	FNamingTokensEvaluationData();
+	
+	/** The current date time. */
+	UPROPERTY(BlueprintReadOnly, Category = "NamingTokens")
+	FDateTime CurrentDateTime;
+};
\ No newline at end of file
diff --git a/Engine/Plugins/VirtualProduction/Takes/Source/TakeRecorder/Private/Recorder/TakeRecorder.cpp b/Engine/Plugins/VirtualProduction/Takes/Source/TakeRecorder/Private/Recorder/TakeRecorder.cpp
--- a/Engine/Plugins/VirtualProduction/Takes/Source/TakeRecorder/Private/Recorder/TakeRecorder.cpp
+++ b/Engine/Plugins/VirtualProduction/Takes/Source/TakeRecorder/Private/Recorder/TakeRecorder.cpp
@@ -455,6 +455,8 @@ bool UTakeRecorder::Initialize ( ULevelSequence* LevelSequenceBase, UTakeRecorde
 		return false;
 	}
 
+	UTakeMetaData::SetMostRecentMetaData(MetaData);
+
 	OnRecordingPreInitializeEvent.Broadcast(this);
 
 	UTakeRecorderBlueprintLibrary::OnTakeRecorderPreInitialize();
@@ -1181,6 +1183,8 @@ void UTakeRecorder::StopInternal(const bool bCancelled)
 	// reset the current recorder and stop us from being ticked
 	if (GetCurrentRecorder().Get() == this)
 	{
+		UTakeMetaData::SetMostRecentMetaData(nullptr);
+		
 		GetCurrentRecorder().Reset();
 		TickableTakeRecorder.WeakRecorder = nullptr;
 
diff --git a/Engine/Plugins/VirtualProduction/Takes/Source/TakeRecorder/Public/Recorder/TakeRecorder.h b/Engine/Plugins/VirtualProduction/Takes/Source/TakeRecorder/Public/Recorder/TakeRecorder.h
--- a/Engine/Plugins/VirtualProduction/Takes/Source/TakeRecorder/Public/Recorder/TakeRecorder.h
+++ b/Engine/Plugins/VirtualProduction/Takes/Source/TakeRecorder/Public/Recorder/TakeRecorder.h
@@ -287,7 +287,7 @@ private:
 	/** Triggered when the movie scene is adjusted. */
 	FOnStartPlayFrameModified OnFrameModifiedEvent;
 
-	/** Sequencer ptr that controls playback of the desination asset during the recording */
+	/** Sequencer ptr that controls playback of the destination asset during the recording */
 	TWeakPtr<ISequencer> WeakSequencer;
 
 	/** Due a few ticks after the pre so we are set up with asset creation */
diff --git a/Engine/Plugins/VirtualProduction/Takes/Source/TakeRecorderNamingTokens/Private/TakeRecoderNamingTokens.cpp b/Engine/Plugins/VirtualProduction/Takes/Source/TakeRecorderNamingTokens/Private/TakeRecoderNamingTokens.cpp
new file mode 100644
--- /dev/null
+++ b/Engine/Plugins/VirtualProduction/Takes/Source/TakeRecorderNamingTokens/Private/TakeRecoderNamingTokens.cpp
@@ -0,0 +1,154 @@
+// Copyright Epic Games, Inc. All Rights Reserved.
+
+#include "TakeRecorderNamingTokens.h"
+
+#include "Editor.h"
+#include "ITakeRecorderNamingTokensModule.h"
+#include "Misc/PackageName.h"
+#include "MovieSceneToolsProjectSettings.h"
+#include "TakeMetaData.h"
+#include "TakeRecorderSettings.h"
+#include "TakesCoreBlueprintLibrary.h"
+
+#define LOCTEXT_NAMESPACE "TakeRecorderNamingTokens"
+
+UTakeRecorderNamingTokens::UTakeRecorderNamingTokens()
+{
+	Namespace = ITakeRecorderNamingTokensModule::GetTakeRecorderNamespace();
+}
+
+UTakeRecorderNamingTokens::~UTakeRecorderNamingTokens()
+{
+}
+
+void UTakeRecorderNamingTokens::OnCreateDefaultTokens(TArray<FNamingTokenData>& Tokens)
+{
+	Super::OnCreateDefaultTokens(Tokens);
+
+	Tokens.Add({
+	    TEXT("day"),
+	    LOCTEXT("TokenDay", "Day"),
+	    FNamingTokenData::FTokenProcessorDelegateNative::CreateLambda([this]{
+	        return FText::FromString(FString::Printf(TEXT("%02i"), GetCurrentDateTime().GetDay()));
+	    })
+	});
+
+	Tokens.Add({
+	    TEXT("month"),
+	    LOCTEXT("TokenMonth", "Month"),
+	    FNamingTokenData::FTokenProcessorDelegateNative::CreateLambda([this]{
+	        return FText::FromString(FString::Printf(TEXT("%02i"), GetCurrentDateTime().GetMonth()));
+	    })
+	});
+
+	Tokens.Add({
+	    TEXT("year"),
+	    LOCTEXT("TokenYear", "Year"),
+	    FNamingTokenData::FTokenProcessorDelegateNative::CreateLambda([this]{
+	        return FText::FromString(FString::Printf(TEXT("%04i"), GetCurrentDateTime().GetYear()));
+	    })
+	});
+
+	Tokens.Add({
+	    TEXT("hour"),
+	    LOCTEXT("TokenHour", "Hour"),
+	    FNamingTokenData::FTokenProcessorDelegateNative::CreateLambda([this]{
+	        return FText::FromString(FString::Printf(TEXT("%02i"), GetCurrentDateTime().GetHour()));
+	    })
+	});
+
+	Tokens.Add({
+	    TEXT("minute"),
+	    LOCTEXT("TokenMinute", "Minute"),
+	    FNamingTokenData::FTokenProcessorDelegateNative::CreateLambda([this]{
+	        return FText::FromString(FString::Printf(TEXT("%02i"), GetCurrentDateTime().GetMinute()));
+	    })
+	});
+
+	Tokens.Add({
+	    TEXT("second"),
+	    LOCTEXT("TokenSecond", "Second"),
+	    FNamingTokenData::FTokenProcessorDelegateNative::CreateLambda([this]{
+	        return FText::FromString(FString::Printf(TEXT("%02i"), GetCurrentDateTime().GetSecond()));
+	    })
+	});
+
+	Tokens.Add({
+	    TEXT("take"),
+	    LOCTEXT("TokenTake", "Take"),
+	    FNamingTokenData::FTokenProcessorDelegateNative::CreateLambda([this]{
+	        const UMovieSceneToolsProjectSettings* ProjectSettings = GetDefault<UMovieSceneToolsProjectSettings>();
+	        const int32 TakeNumDigits = ProjectSettings->TakeNumDigits;
+	        return FText::FromString(FString::Printf(TEXT("%0*d"), TakeNumDigits, TakeMetaData->GetTakeNumber()));
+	    })
+	});
+
+	Tokens.Add({
+	    TEXT("slate"),
+	    LOCTEXT("TokenSlate", "Slate"),
+	    FNamingTokenData::FTokenProcessorDelegateNative::CreateLambda([this]{
+	        return FText::FromString(TakeMetaData->GetSlate());
+	    })
+	});
+
+	Tokens.Add({
+	    TEXT("map"),
+	    LOCTEXT("TokenMap", "Map"),
+	    FNamingTokenData::FTokenProcessorDelegateNative::CreateLambda([]{
+	        FString MapName;
+	        if (GIsEditor)
+	        {
+	            MapName = FPackageName::GetShortFName(
+	                GEditor->GetEditorWorldContext().World()->PersistentLevel->GetOutermost()->GetFName()
+	            ).GetPlainNameString();
+	        }
+	        return FText::FromString(MapName);
+	    })
+	});
+}
+
+void UTakeRecorderNamingTokens::OnPreEvaluate_Implementation(const FNamingTokensEvaluationData& InEvaluationData)
+{
+	Super::OnPreEvaluate_Implementation(InEvaluationData);
+	
+	TakeMetaData = UTakeMetaData::GetMostRecentMetaData();
+	if (TakeMetaData == nullptr)
+	{
+		// Create metadata which computes its data from available information.
+		// This is how STakeRecorderCockpit handles cases when metadata isn't available. Ideally this
+		// method would be available for both us and the slate widget, but UTakeRecorderProjectSettings creates
+		// a circular dependency we should avoid.
+		
+		TakeMetaData = UTakeMetaData::CreateFromDefaults(GetTransientPackage(), NAME_None);
+		TakeMetaData->SetFlags(RF_Transactional | RF_Transient);
+
+		const FString DefaultSlate = GetDefault<UTakeRecorderProjectSettings>()->Settings.DefaultSlate;
+		if (TakeMetaData->GetSlate() != DefaultSlate)
+		{
+			TakeMetaData->SetSlate(DefaultSlate, false);
+		}
+
+		// Compute the correct starting take number
+		const int32 NextTakeNumber = UTakesCoreBlueprintLibrary::ComputeNextTakeNumber(TakeMetaData->GetSlate());
+		if (TakeMetaData->GetTakeNumber() != NextTakeNumber)
+		{
+			TakeMetaData->SetTakeNumber(NextTakeNumber, false);
+		}
+	}
+}
+
+void UTakeRecorderNamingTokens::OnPostEvaluate_Implementation()
+{
+	Super::OnPostEvaluate_Implementation();
+
+	// Make sure we don't keep a strong reference in case we're linking to active metadata.
+	TakeMetaData = nullptr;
+}
+
+FDateTime UTakeRecorderNamingTokens::GetCurrentDateTime_Implementation() const
+{
+	const FDateTime TakeTimestamp = TakeMetaData ? TakeMetaData->GetTimestamp() : FDateTime(0);
+	return TakeTimestamp == FDateTime(0) ? Super::GetCurrentDateTime_Implementation() : TakeTimestamp;
+}
+
+#undef LOCTEXT_NAMESPACE
diff --git a/Engine/Plugins/VirtualProduction/Takes/Source/TakeRecorderNamingTokens/Private/TakeRecorderNamingTokens.h b/Engine/Plugins/VirtualProduction/Takes/Source/TakeRecorderNamingTokens/Private/TakeRecorderNamingTokens.h
new file mode 100644
--- /dev/null
+++ b/Engine/Plugins/VirtualProduction/Takes/Source/TakeRecorderNamingTokens/Private/TakeRecorderNamingTokens.h
@@ -0,0 +1,35 @@
+// Copyright Epic Games, Inc. All Rights Reserved.
+
+#pragma once
+
+#include "NamingTokens.h"
+
+#include "TakeRecorderNamingTokens.generated.h"
+
+class UTakeMetaData;
+
+/**
+ * Naming Tokens for Take Recorder.
+ */
+UCLASS(NotBlueprintable)
+class UTakeRecorderNamingTokens final: public UNamingTokens
+{
+	GENERATED_BODY()
+
+public:
+	UTakeRecorderNamingTokens();
+	virtual ~UTakeRecorderNamingTokens() override;
+
+protected:
+	// ~Begin UNamingTokens
+	virtual void OnCreateDefaultTokens(TArray<FNamingTokenData>& Tokens) override;
+	virtual void OnPreEvaluate_Implementation(const FNamingTokensEvaluationData& InEvaluationData) override;
+	virtual void OnPostEvaluate_Implementation() override;
+	virtual FDateTime GetCurrentDateTime_Implementation() const override;
+	// ~End UNamingTokens
+	
+private:
+	/** Cached metadata for this run. */
+	UPROPERTY(Transient)
+	TObjectPtr<UTakeMetaData> TakeMetaData;
+};
diff --git a/Engine/Plugins/VirtualProduction/Takes/Source/TakeRecorderNamingTokens/Private/TakeRecorderNamingTokensLog.h b/Engine/Plugins/VirtualProduction/Takes/Source/TakeRecorderNamingTokens/Private/TakeRecorderNamingTokensLog.h
new file mode 100644
--- /dev/null
+++ b/Engine/Plugins/VirtualProduction/Takes/Source/TakeRecorderNamingTokens/Private/TakeRecorderNamingTokensLog.h
@@ -0,0 +1,7 @@
+// Copyright Epic Games, Inc. All Rights Reserved.
+
+#pragma once
+
+#include "Logging/LogMacros.h"
+
+DECLARE_LOG_CATEGORY_EXTERN(LogTakeRecorderNamingTokens, Log, All);
\ No newline at end of file
diff --git a/Engine/Plugins/VirtualProduction/Takes/Source/TakeRecorderNamingTokens/Private/TakeRecorderNamingTokensModule.cpp b/Engine/Plugins/VirtualProduction/Takes/Source/TakeRecorderNamingTokens/Private/TakeRecorderNamingTokensModule.cpp
new file mode 100644
--- /dev/null
+++ b/Engine/Plugins/VirtualProduction/Takes/Source/TakeRecorderNamingTokens/Private/TakeRecorderNamingTokensModule.cpp
@@ -0,0 +1,18 @@
+// Copyright Epic Games, Inc. All Rights Reserved.
+
+#include "ITakeRecorderNamingTokensModule.h"
+#include "TakeRecorderNamingTokensLog.h"
+
+class FTakeRecorderNamingTokensModule : public ITakeRecorderNamingTokensModule
+{
+	virtual void StartupModule() override
+	{
+	}
+
+	virtual void ShutdownModule() override
+	{
+	}
+};
+
+IMPLEMENT_MODULE(FTakeRecorderNamingTokensModule, TakeRecorderNamingTokens);
+DEFINE_LOG_CATEGORY(LogTakeRecorderNamingTokens);
\ No newline at end of file
diff --git a/Engine/Plugins/VirtualProduction/Takes/Source/TakeRecorderNamingTokens/Public/ITakeRecorderNamingTokensModule.h b/Engine/Plugins/VirtualProduction/Takes/Source/TakeRecorderNamingTokens/Public/ITakeRecorderNamingTokensModule.h
new file mode 100644
--- /dev/null
+++ b/Engine/Plugins/VirtualProduction/Takes/Source/TakeRecorderNamingTokens/Public/ITakeRecorderNamingTokensModule.h
@@ -0,0 +1,15 @@
+﻿// Copyright Epic Games, Inc. All Rights Reserved.
+
+#pragma once
+
+#include "Modules/ModuleManager.h"
+
+class ITakeRecorderNamingTokensModule : public IModuleInterface
+{
+public:
+	/** The Take Recorder namespace used for processing naming tokens. */
+	static FString GetTakeRecorderNamespace()
+	{
+		return TEXT("tr");
+	}
+};
diff --git a/Engine/Plugins/VirtualProduction/Takes/Source/TakeRecorderNamingTokens/TakeRecorderNamingTokens.Build.cs b/Engine/Plugins/VirtualProduction/Takes/Source/TakeRecorderNamingTokens/TakeRecorderNamingTokens.Build.cs
new file mode 100644
--- /dev/null
+++ b/Engine/Plugins/VirtualProduction/Takes/Source/TakeRecorderNamingTokens/TakeRecorderNamingTokens.Build.cs
@@ -0,0 +1,27 @@
+// Copyright Epic Games, Inc. All Rights Reserved.
+
+using UnrealBuildTool;
+
+public class TakeRecorderNamingTokens : ModuleRules
+{
+	public TakeRecorderNamingTokens(ReadOnlyTargetRules Target) : base(Target)
+	{
+		PrivateDependencyModuleNames.AddRange(
+			new string[] {
+				"Core",
+				"CoreUObject",
+				"UnrealEd",
+				"MovieSceneTools"
+			}
+		);
+
+		PublicDependencyModuleNames.AddRange(
+			new string[] {
+				"Engine",
+				"TakeRecorder",
+				"TakesCore",
+				"NamingTokens"
+			}
+		);
+	}
+}
diff --git a/Engine/Plugins/VirtualProduction/Takes/Source/TakesCore/Private/TakeMetaData.cpp b/Engine/Plugins/VirtualProduction/Takes/Source/TakesCore/Private/TakeMetaData.cpp
--- a/Engine/Plugins/VirtualProduction/Takes/Source/TakesCore/Private/TakeMetaData.cpp
+++ b/Engine/Plugins/VirtualProduction/Takes/Source/TakesCore/Private/TakeMetaData.cpp
@@ -1,14 +1,17 @@
 // Copyright Epic Games, Inc. All Rights Reserved.
 
 #include "TakeMetaData.h"
+
+#include "Editor.h"
+#include "Engine/Level.h"
+#include "ITakeRecorderNamingTokensModule.h"
+#include "NamingTokenData.h"
+#include "NamingTokensEditorSubsystem.h"
 #include "TakePreset.h"
+#include "TakesCoreBlueprintLibrary.h"
+#include "TakesCoreLog.h"
 #include "UObject/AssetRegistryTagsContext.h"
 #include "UObject/Package.h"
-#include "TakesCoreBlueprintLibrary.h"
-#include "ObjectTools.h"
-#include "Engine/Level.h"
-#include "Editor.h"
-#include "MovieSceneToolsProjectSettings.h"
 
 #include UE_INLINE_GENERATED_CPP_BY_NAME(TakeMetaData)
 
@@ -49,6 +52,71 @@ UTakeMetaData* UTakeMetaData::CreateFromDefaults(UObject* Outer, FName Name)
 	return CastChecked<UTakeMetaData>(StaticDuplicateObject(GetConfigInstance(), Outer, Name, RF_NoFlags));
 }
 
+namespace UE::TakeRecorder::MetaDataHelpers::Private
+{
+	/** A copy of the last metadata. */
+	static TStrongObjectPtr<UTakeMetaData>& GetLastMetaData()
+	{
+		static TStrongObjectPtr<UTakeMetaData> CachedMetaData;
+		return CachedMetaData;
+	}
+
+	/** Weak reference to the current metadata. */
+	static TWeakObjectPtr<UTakeMetaData>& GetCurrentWeakMetaData()
+	{
+		static TWeakObjectPtr<UTakeMetaData> CachedWeakMetaData;
+		return CachedWeakMetaData;
+	}
+
+	/** Updates the cached metadata. */
+	static void UpdateCachedMetaData(UTakeMetaData* NewMetaData)
+	{
+		if (NewMetaData)
+		{
+			GetCurrentWeakMetaData() = NewMetaData;
+			GetLastMetaData().Reset();
+		}
+		else
+		{
+			// We want to copy the metadata object so the original can be cleaned up if the level sequence needs to be removed.
+			UTakeMetaData* MetaDataCopy = GetCurrentWeakMetaData().IsValid() ?
+				DuplicateObject(GetCurrentWeakMetaData().Get(), GetTransientPackage()) : nullptr;
+			GetLastMetaData().Reset(MetaDataCopy);
+			GetCurrentWeakMetaData().Reset();
+		}
+	}
+
+	/** Update the active metadata for this scope. */
+	struct FCacheMetaDataOnScope
+	{
+		explicit FCacheMetaDataOnScope(UTakeMetaData* InMetaData)
+		{
+			UpdateCachedMetaData(InMetaData);
+		}
+		~FCacheMetaDataOnScope()
+		{
+			UpdateCachedMetaData(nullptr);
+		}
+	};
+}
+
+UTakeMetaData* UTakeMetaData::GetMostRecentMetaData()
+{
+	using namespace UE::TakeRecorder::MetaDataHelpers::Private;
+	if (UTakeMetaData* CurrentMetaData = GetCurrentWeakMetaData().Get())
+	{
+		return CurrentMetaData;
+	}
+
+	return GetLastMetaData().Get();
+}
+
+void UTakeMetaData::SetMostRecentMetaData(UTakeMetaData* InMetaData)
+{
+	using namespace UE::TakeRecorder::MetaDataHelpers::Private;
+	UpdateCachedMetaData(InMetaData);
+}
+
 bool UTakeMetaData::Recorded() const
 {
 	return bool( GetTimestamp() != FDateTime(0) );
@@ -71,33 +139,20 @@ void UTakeMetaData::Unlock()
 
 FString UTakeMetaData::GenerateAssetPath(const FString& PathFormatString) const
 {
-	FDateTime LocaleNow = FDateTime::Now();
-	FDateTime TimestampToUse = Timestamp == FDateTime(0) ? LocaleNow : Timestamp;
-
-	FString MapName;
-
-#if WITH_EDITOR
-	if (GIsEditor)
+	if (!ensure(GIsEditor))
 	{
-		MapName = FPackageName::GetShortFName(GEditor->GetEditorWorldContext().World()->PersistentLevel->GetOutermost()->GetFName()).GetPlainNameString();
+		UE_LOG(LogTakesCore, Warning, TEXT("GIsEditor is false during UTakeMetaData::GenerateAssetPath and the generated path will be incorrect."))
+		return FString();
 	}
-#endif
-
-	const UMovieSceneToolsProjectSettings* ProjectSettings = GetDefault<UMovieSceneToolsProjectSettings>();
-	const int32 TakeNumDigits = ProjectSettings->TakeNumDigits;
-
-	TMap<FString, FStringFormatArg> FormatArgs;
-	FormatArgs.Add(TEXT("day"),    FString::Printf(TEXT("%02i"), TimestampToUse.GetDay()));
-	FormatArgs.Add(TEXT("month"),  FString::Printf(TEXT("%02i"), TimestampToUse.GetMonth()));
-	FormatArgs.Add(TEXT("year"),   FString::Printf(TEXT("%04i"), TimestampToUse.GetYear()));
-	FormatArgs.Add(TEXT("hour"),   FString::Printf(TEXT("%02i"), TimestampToUse.GetHour()));
-	FormatArgs.Add(TEXT("minute"), FString::Printf(TEXT("%02i"), TimestampToUse.GetMinute()));
-	FormatArgs.Add(TEXT("second"), FString::Printf(TEXT("%02i"), TimestampToUse.GetSecond()));
-	FormatArgs.Add(TEXT("take"),   FString::Printf(TEXT("%0*d"), TakeNumDigits, TakeNumber));
-	FormatArgs.Add(TEXT("slate"),  *Slate);
-	FormatArgs.Add(TEXT("map"),    *MapName);
-
-	return FString::Format(*PathFormatString, FormatArgs);
+
+	UE::TakeRecorder::MetaDataHelpers::Private::FCacheMetaDataOnScope ScopedActiveMetaData(const_cast<UTakeMetaData*>(this));
+	
+	FNamingTokenFilterArgs NamingTokenFilters;
+	NamingTokenFilters.NamespacesToInclude = { ITakeRecorderNamingTokensModule::GetTakeRecorderNamespace() };
+	const FNamingTokenResultData Result = GEditor->GetEditorSubsystem<UNamingTokensEditorSubsystem>()->EvaluateTokenString(PathFormatString,
+		NamingTokenFilters);
+
+	return Result.EvaluatedText.ToString();
 }
 
 const FString& UTakeMetaData::GetSlate() const
diff --git a/Engine/Plugins/VirtualProduction/Takes/Source/TakesCore/Public/TakeMetaData.h b/Engine/Plugins/VirtualProduction/Takes/Source/TakesCore/Public/TakeMetaData.h
--- a/Engine/Plugins/VirtualProduction/Takes/Source/TakesCore/Public/TakeMetaData.h
+++ b/Engine/Plugins/VirtualProduction/Takes/Source/TakesCore/Public/TakeMetaData.h
@@ -58,7 +58,17 @@ public:
 	 * @param Name     The name for the new object. Must not already exist
 	 */
 	static UTakeMetaData* CreateFromDefaults(UObject* Outer, FName Name);
+	
+	/**
+	 * Retrieve the most recent or active take metadata.
+	 */
+	static UTakeMetaData* GetMostRecentMetaData();
 
+	/**
+	 * Update the most recent accessed metadata.
+	 */
+	static void SetMostRecentMetaData(UTakeMetaData* InMetaData);
+	
 public:
 
 	/**
diff --git a/Engine/Plugins/VirtualProduction/Takes/Source/TakesCore/TakesCore.Build.cs b/Engine/Plugins/VirtualProduction/Takes/Source/TakesCore/TakesCore.Build.cs
--- a/Engine/Plugins/VirtualProduction/Takes/Source/TakesCore/TakesCore.Build.cs
+++ b/Engine/Plugins/VirtualProduction/Takes/Source/TakesCore/TakesCore.Build.cs
@@ -12,6 +12,7 @@ public class TakesCore : ModuleRules
 				"CoreUObject",
 				"EditorFramework",
 				"UnrealEd",
+				"NamingTokens",
 			}
 		);
 
@@ -30,5 +31,12 @@ public class TakesCore : ModuleRules
 
 			}
 		);
+		
+		PublicIncludePathModuleNames.AddRange(
+			new string[]
+			{
+				"TakeRecorderNamingTokens"
+			}
+		);
 	}
 }
diff --git a/Engine/Plugins/VirtualProduction/Takes/Takes.uplugin b/Engine/Plugins/VirtualProduction/Takes/Takes.uplugin
--- a/Engine/Plugins/VirtualProduction/Takes/Takes.uplugin
+++ b/Engine/Plugins/VirtualProduction/Takes/Takes.uplugin
@@ -23,6 +23,10 @@
 		{
 			"Name": "LevelSequenceEditor",
 			"Enabled": true
+		},
+		{
+			"Name": "NamingTokens",
+			"Enabled": true
 		}
 	],
 	"Modules": 
@@ -68,6 +72,12 @@
 			"ProgramAllowList": [ "LiveLinkHub" ],
 			"Type": "UncookedOnly",
 			"LoadingPhase": "PostDefault"
+		},
+		{
+			"Name": "TakeRecorderNamingTokens",
+			"ProgramAllowList": [ "LiveLinkHub" ],
+			"Type": "UncookedOnly",
+			"LoadingPhase": "PostDefault"
 		}
 	]
 }
\ No newline at end of file
