NamingTokens: Fix shutdown module errors.

diff --git a/Engine/Plugins/Developer/NamingTokens/Source/NamingTokens/Private/NamingTokensModule.cpp b/Engine/Plugins/Developer/NamingTokens/Source/NamingTokens/Private/NamingTokensModule.cpp
--- a/Engine/Plugins/Developer/NamingTokens/Source/NamingTokens/Private/NamingTokensModule.cpp
+++ b/Engine/Plugins/Developer/NamingTokens/Source/NamingTokens/Private/NamingTokensModule.cpp
@@ -23,9 +23,14 @@ void FNamingTokensModule::StartupModule()
 
 void FNamingTokensModule::ShutdownModule()
 {
-	FPropertyEditorModule& PropertyModule = FModuleManager::GetModuleChecked<FPropertyEditorModule>("PropertyEditor");
-	PropertyModule.UnregisterCustomPropertyTypeLayout(FNamingTokenData::StaticStruct()->GetFName());
-
+	if (UObjectInitialized() && !IsEngineExitRequested())
+	{
+		if (FPropertyEditorModule* PropertyEditorModule = FModuleManager::GetModulePtr<FPropertyEditorModule>("PropertyEditor"))
+		{
+			PropertyEditorModule->UnregisterCustomPropertyTypeLayout(FNamingTokenData::StaticStruct()->GetFName());
+		}
+	}
+	
 	if (GEditor)
 	{
 		if (UNamingTokensEditorSubsystem* NamingTokensSubsystem = GEditor->GetEditorSubsystem<UNamingTokensEditorSubsystem>())
