NamingTokens: Add asset definition and factory. Naming Tokens can be created under the Misc menu.

diff --git a/Engine/Plugins/Developer/NamingTokens/NamingTokens.uplugin b/Engine/Plugins/Developer/NamingTokens/NamingTokens.uplugin
--- a/Engine/Plugins/Developer/NamingTokens/NamingTokens.uplugin
+++ b/Engine/Plugins/Developer/NamingTokens/NamingTokens.uplugin
@@ -19,5 +19,11 @@
 			"Type": "UncookedOnly",
 			"LoadingPhase": "Default"
 		}
+	],
+	"Plugins": [
+		{
+			"Name": "EngineAssetDefinitions",
+			"Enabled": true
+		}
 	]
 }
\ No newline at end of file
diff --git a/Engine/Plugins/Developer/NamingTokens/Source/NamingTokens/NamingTokens.Build.cs b/Engine/Plugins/Developer/NamingTokens/Source/NamingTokens/NamingTokens.Build.cs
--- a/Engine/Plugins/Developer/NamingTokens/Source/NamingTokens/NamingTokens.Build.cs
+++ b/Engine/Plugins/Developer/NamingTokens/Source/NamingTokens/NamingTokens.Build.cs
@@ -14,7 +14,7 @@ public class NamingTokens : ModuleRules
 				"Core",
 				"EditorSubsystem"
 			}
-			);
+		);
 		
 		PrivateDependencyModuleNames.AddRange(
 			new string[]
@@ -26,7 +26,9 @@ public class NamingTokens : ModuleRules
 				"InputCore",
 				"UnrealEd",
 				"BlueprintGraph",
+				"AssetDefinition",
+				"EngineAssetDefinitions",
 			}
-			);
+		);
 	}
 }
diff --git a/Engine/Plugins/Developer/NamingTokens/Source/NamingTokens/Private/AssetDefinition_NamingTokens.cpp b/Engine/Plugins/Developer/NamingTokens/Source/NamingTokens/Private/AssetDefinition_NamingTokens.cpp
new file mode 100644
--- /dev/null
+++ b/Engine/Plugins/Developer/NamingTokens/Source/NamingTokens/Private/AssetDefinition_NamingTokens.cpp
@@ -0,0 +1,21 @@
+// Copyright Epic Games, Inc. All Rights Reserved.
+
+#include "AssetDefinition_NamingTokens.h"
+
+#include "NamingTokens.h"
+
+TSoftClassPtr<UObject> UAssetDefinition_NamingTokens::GetAssetClass() const
+{
+	return UNamingTokens::StaticClass();
+}
+
+TConstArrayView<FAssetCategoryPath> UAssetDefinition_NamingTokens::GetAssetCategories() const
+{
+	static const FAssetCategoryPath Categories[] = { FAssetCategoryPath(EAssetCategoryPaths::Misc) };
+	return Categories;
+}
+
+FText UAssetDefinition_NamingTokens::GetAssetDisplayName() const
+{
+	return NSLOCTEXT("AssetDefinition_NamingTokens", "AssetDefinition_NamingTokens_DisplayName", "Naming Tokens");
+}
diff --git a/Engine/Plugins/Developer/NamingTokens/Source/NamingTokens/Private/AssetDefinition_NamingTokens.h b/Engine/Plugins/Developer/NamingTokens/Source/NamingTokens/Private/AssetDefinition_NamingTokens.h
new file mode 100644
--- /dev/null
+++ b/Engine/Plugins/Developer/NamingTokens/Source/NamingTokens/Private/AssetDefinition_NamingTokens.h
@@ -0,0 +1,18 @@
+// Copyright Epic Games, Inc. All Rights Reserved.
+
+#pragma once
+
+#include "Script/AssetDefinition_Blueprint.h"
+
+#include "AssetDefinition_NamingTokens.generated.h"
+
+UCLASS()
+class UAssetDefinition_NamingTokens : public UAssetDefinition_Blueprint
+{
+	GENERATED_BODY()
+	
+public:
+	virtual TSoftClassPtr<UObject> GetAssetClass() const override;
+	virtual TConstArrayView<FAssetCategoryPath> GetAssetCategories() const override;
+	virtual FText GetAssetDisplayName() const override;
+};
diff --git a/Engine/Plugins/Developer/NamingTokens/Source/NamingTokens/Private/NamingTokens.cpp b/Engine/Plugins/Developer/NamingTokens/Source/NamingTokens/Private/NamingTokens.cpp
--- a/Engine/Plugins/Developer/NamingTokens/Source/NamingTokens/Private/NamingTokens.cpp
+++ b/Engine/Plugins/Developer/NamingTokens/Source/NamingTokens/Private/NamingTokens.cpp
@@ -123,6 +123,11 @@ const TArray<FNamingTokenData>& UNamingTokens::GetDefaultTokens() const
 	return DefaultTokens;
 }
 
+const TArray<FNamingTokenData>& UNamingTokens::GetCustomTokens() const
+{
+	return CustomTokens;
+}
+
 FDateTime UNamingTokens::GetCurrentDateTime_Implementation() const
 {
 	return CurrentEvaluationData.CurrentDateTime;
diff --git a/Engine/Plugins/Developer/NamingTokens/Source/NamingTokens/Private/NamingTokensFactory.cpp b/Engine/Plugins/Developer/NamingTokens/Source/NamingTokens/Private/NamingTokensFactory.cpp
new file mode 100644
--- /dev/null
+++ b/Engine/Plugins/Developer/NamingTokens/Source/NamingTokens/Private/NamingTokensFactory.cpp
@@ -0,0 +1,67 @@
+// Copyright Epic Games, Inc. All Rights Reserved.
+
+#include "NamingTokensFactory.h"
+
+#include "NamingTokens.h"
+
+#include "BlueprintEditorSettings.h"
+#include "EdGraphSchema_K2.h"
+#include "Engine/Blueprint.h"
+#include "Engine/BlueprintGeneratedClass.h"
+#include "K2Node_Event.h"
+#include "Kismet2/KismetEditorUtilities.h"
+
+UNamingTokensFactory::UNamingTokensFactory()
+{
+	SupportedClass = UNamingTokens::StaticClass();
+
+	bCreateNew = true;
+	bEditAfterNew = true;
+}
+
+UObject* UNamingTokensFactory::FactoryCreateNew(UClass* InClass, UObject* InParent, FName InName, EObjectFlags Flags, UObject* Context,
+	FFeedbackContext* Warn)
+{
+	UClass* ParentClass = SupportedClass;
+	check(InClass->IsChildOf(ParentClass));
+    
+	UBlueprint* NewBP = FKismetEditorUtilities::CreateBlueprint(
+		ParentClass,
+		InParent,
+		InName,
+		BPTYPE_Normal,
+		UBlueprint::StaticClass(),
+		UBlueprintGeneratedClass::StaticClass()
+	);
+	
+	const UBlueprintEditorSettings* Settings = GetDefault<UBlueprintEditorSettings>();
+	if (Settings && Settings->bSpawnDefaultBlueprintNodes)
+	{
+		// Create default events.
+		int32 NodePositionY = 0;
+
+		UEdGraph* EventGraph = FindObject<UEdGraph>(NewBP, *(UEdGraphSchema_K2::GN_EventGraph.ToString()));
+		check(EventGraph);
+		
+		if (UK2Node_Event* OnPreEvaluateNode = FKismetEditorUtilities::AddDefaultEventNode(NewBP, EventGraph,
+			GET_FUNCTION_NAME_CHECKED(UNamingTokens, OnPreEvaluate), UNamingTokens::StaticClass(), NodePositionY))
+		{
+			// Set the node comment. The comment is displayed because we are automatically placed as ghost nodes.
+			// Once a connection is made the comment will go away, but is still accessible via tooltip.
+			const UFunction* Function = UNamingTokens::StaticClass()->FindFunctionByName(GET_FUNCTION_NAME_CHECKED(UNamingTokens, OnPreEvaluate));
+			check(Function);
+			OnPreEvaluateNode->NodeComment = Function->GetToolTipText().ToString();
+		}
+		if (UK2Node_Event* OnPostEvaluateNode = FKismetEditorUtilities::AddDefaultEventNode(NewBP, EventGraph,
+			GET_FUNCTION_NAME_CHECKED(UNamingTokens, OnPostEvaluate), UNamingTokens::StaticClass(), NodePositionY))
+		{
+			// Set the node comment. The comment is displayed because we are automatically placed as ghost nodes.
+			// Once a connection is made the comment will go away, but is still accessible via tooltip.
+			const UFunction* Function = UNamingTokens::StaticClass()->FindFunctionByName(GET_FUNCTION_NAME_CHECKED(UNamingTokens, OnPostEvaluate));
+			check(Function);
+			OnPostEvaluateNode->NodeComment = Function->GetToolTipText().ToString();
+		}
+	}
+
+	return NewBP;
+}
diff --git a/Engine/Plugins/Developer/NamingTokens/Source/NamingTokens/Private/NamingTokensFactory.h b/Engine/Plugins/Developer/NamingTokens/Source/NamingTokens/Private/NamingTokensFactory.h
new file mode 100644
--- /dev/null
+++ b/Engine/Plugins/Developer/NamingTokens/Source/NamingTokens/Private/NamingTokensFactory.h
@@ -0,0 +1,19 @@
+// Copyright Epic Games, Inc. All Rights Reserved.
+
+#pragma once
+
+#include "CoreMinimal.h"
+#include "Factories/Factory.h"
+
+#include "NamingTokensFactory.generated.h"
+
+UCLASS(MinimalAPI)
+class UNamingTokensFactory : public UFactory
+{
+	GENERATED_BODY()
+
+public:
+	UNamingTokensFactory();
+	
+	virtual UObject* FactoryCreateNew(UClass* InClass, UObject* InParent, FName InName, EObjectFlags Flags, UObject* Context, FFeedbackContext* Warn) override;
+};
\ No newline at end of file
diff --git a/Engine/Plugins/Developer/NamingTokens/Source/NamingTokens/Public/NamingTokens.h b/Engine/Plugins/Developer/NamingTokens/Source/NamingTokens/Public/NamingTokens.h
--- a/Engine/Plugins/Developer/NamingTokens/Source/NamingTokens/Public/NamingTokens.h
+++ b/Engine/Plugins/Developer/NamingTokens/Source/NamingTokens/Public/NamingTokens.h
@@ -47,6 +47,9 @@ public:
 	/** Retrieve the default tokens. */
 	const TArray<FNamingTokenData>& GetDefaultTokens() const;
 	
+	/** Retrieve the custom tokens. */
+	const TArray<FNamingTokenData>& GetCustomTokens() const;
+	
 	/** Retrieve the naming token's namespace. */
 	const FString& GetNamespace() const { return Namespace; }
 	/** The name of the namespace property. */
@@ -81,7 +84,7 @@ protected:
 	void PostEvaluate();
 	
 	/**
-	 * Called prior to evaluation. Allows consistent data to be setup for each token evaluation.
+	 * Called prior to evaluation. Allows consistent data to be set up for each token evaluation.
 	 * 
 	 * This is important if the data is temporally sensitive and could change between evaluating
 	 * multiple tokens in a string, such as a high resolution timer.
@@ -134,4 +137,6 @@ private:
 	UPROPERTY(VisibleDefaultsOnly, Transient, Category = "NamingTokens")
 	FText TestTokenResult;
 #endif
+
+	friend class UNamingTokensFactory;
 };
