NamingTokens: Make plugin experimental.

diff --git a/Engine/Plugins/Developer/NamingTokens/NamingTokens.uplugin b/Engine/Plugins/Developer/NamingTokens/NamingTokens.uplugin
--- a/Engine/Plugins/Developer/NamingTokens/NamingTokens.uplugin
+++ b/Engine/Plugins/Developer/NamingTokens/NamingTokens.uplugin
@@ -10,8 +10,8 @@
 	"DocsURL": "",
 	"MarketplaceURL": "",
 	"CanContainContent": true,
-	"IsBetaVersion": true,
-	"IsExperimentalVersion": false,
+	"IsBetaVersion": false,
+	"IsExperimentalVersion": true,
 	"Installed": false,
 	"Modules": [
 		{
