NamingTokens: Change from editor subsystem to an engine subsystem.

diff --git a/Engine/Plugins/Animation/LiveLinkHub/Source/LiveLinkHub/Private/Recording/Implementations/LiveLinkUAssetRecorder.cpp b/Engine/Plugins/Animation/LiveLinkHub/Source/LiveLinkHub/Private/Recording/Implementations/LiveLinkUAssetRecorder.cpp
--- a/Engine/Plugins/Animation/LiveLinkHub/Source/LiveLinkHub/Private/Recording/Implementations/LiveLinkUAssetRecorder.cpp
+++ b/Engine/Plugins/Animation/LiveLinkHub/Source/LiveLinkHub/Private/Recording/Implementations/LiveLinkUAssetRecorder.cpp
@@ -7,6 +7,7 @@
 #include "Config/LiveLinkHubFileUtilities.h"
 #include "Containers/Map.h"
 #include "ContentBrowserModule.h"
+#include "Engine/Engine.h"
 #include "Features/IModularFeatures.h"
 #include "FileHelpers.h"
 #include "Framework/MultiBox/MultiBoxBuilder.h"
@@ -22,7 +23,7 @@
 #include "Misc/FileHelper.h"
 #include "Misc/PackageName.h"
 #include "Misc/ScopeExit.h"
-#include "NamingTokensEditorSubsystem.h"
+#include "NamingTokensEngineSubsystem.h"
 #include "Recording/Implementations/LiveLinkUAssetRecording.h"
 #include "Recording/LiveLinkHubPlaybackController.h"
 #include "Recording/LiveLinkRecording.h"
@@ -210,7 +211,7 @@ bool FLiveLinkUAssetRecorder::GetSavePresetPackageName(FString& OutName)
 	FNamingTokenFilterArgs Filter;
 	Filter.NamespacesToInclude.Add(Tokens->GetNamespace());
 	check(GEditor);
-	FNamingTokenResultData TemplateData = GEditor->GetEditorSubsystem<UNamingTokensEditorSubsystem>()->EvaluateTokenString(GetDefault<ULiveLinkHubSettings>()->FilenameTemplate,
+	FNamingTokenResultData TemplateData = GEngine->GetEngineSubsystem<UNamingTokensEngineSubsystem>()->EvaluateTokenString(GetDefault<ULiveLinkHubSettings>()->FilenameTemplate,
 		Filter);
 
 	FString DefaultName;
diff --git a/Engine/Plugins/Animation/LiveLinkHub/Source/LiveLinkHub/Private/Settings/LiveLinkHubSettings.cpp b/Engine/Plugins/Animation/LiveLinkHub/Source/LiveLinkHub/Private/Settings/LiveLinkHubSettings.cpp
--- a/Engine/Plugins/Animation/LiveLinkHub/Source/LiveLinkHub/Private/Settings/LiveLinkHubSettings.cpp
+++ b/Engine/Plugins/Animation/LiveLinkHub/Source/LiveLinkHub/Private/Settings/LiveLinkHubSettings.cpp
@@ -3,9 +3,9 @@
 #include "LiveLinkHubSettings.h"
 
 #include "Config/LiveLinkHubTemplateTokens.h"
-#include "NamingTokensEditorSubsystem.h"
+#include "NamingTokensEngineSubsystem.h"
 
-#include "Editor.h"
+#include "Engine/Engine.h"
 #include "UObject/Package.h"
 
 void ULiveLinkHubSettings::PostEditChangeProperty(FPropertyChangedEvent& PropertyChangedEvent)
@@ -33,8 +33,8 @@ void ULiveLinkHubSettings::CalculateExampleOutput()
 		FNamingTokenFilterArgs Filter;
 		Filter.NamespacesToInclude.Add(Tokens->GetNamespace());
 		
-		check(GEditor);
-		const FNamingTokenResultData TemplateData = GEditor->GetEditorSubsystem<UNamingTokensEditorSubsystem>()->EvaluateTokenString(
+		check(GEngine);
+		const FNamingTokenResultData TemplateData = GEngine->GetEngineSubsystem<UNamingTokensEngineSubsystem>()->EvaluateTokenString(
 			GetDefault<ULiveLinkHubSettings>()->FilenameTemplate, Filter);
 		FilenameOutput = TemplateData.EvaluatedText.ToString();
 	}
diff --git a/Engine/Plugins/Animation/LiveLinkHub/Source/LiveLinkHub/Private/Settings/LiveLinkHubSettingsCustomization.cpp b/Engine/Plugins/Animation/LiveLinkHub/Source/LiveLinkHub/Private/Settings/LiveLinkHubSettingsCustomization.cpp
--- a/Engine/Plugins/Animation/LiveLinkHub/Source/LiveLinkHub/Private/Settings/LiveLinkHubSettingsCustomization.cpp
+++ b/Engine/Plugins/Animation/LiveLinkHub/Source/LiveLinkHub/Private/Settings/LiveLinkHubSettingsCustomization.cpp
@@ -5,9 +5,9 @@
 #include "Config/LiveLinkHubTemplateTokens.h"
 #include "DetailLayoutBuilder.h"
 #include "DetailWidgetRow.h"
-#include "Editor.h"
+#include "Engine/Engine.h"
 #include "LiveLinkHubSettings.h"
-#include "NamingTokensEditorSubsystem.h"
+#include "NamingTokensEngineSubsystem.h"
 #include "Widgets/Input/SEditableText.h"
 #include "Widgets/Text/STextBlock.h"
 
@@ -55,14 +55,14 @@ void FLiveLinkHubSettingsCustomization::CustomizeDetails(IDetailLayoutBuilder& D
 FText FLiveLinkHubSettingsCustomization::GetDisplayTokenText() const
 {
 	FString FormattedTokensString = TEXT("None");
-	if (GEditor)
+	if (GEngine)
 	{
 		FNamingTokenFilterArgs Args;
 		if (const TObjectPtr<ULiveLinkHubNamingTokens> Tokens = GetDefault<ULiveLinkHubSettings>()->GetNamingTokens())
 		{
 			Args.NamespacesToInclude.Add(Tokens->GetNamespace());
 		}
-		FormattedTokensString = GEditor->GetEditorSubsystem<UNamingTokensEditorSubsystem>()->GetFormattedTokensStringForDisplay(MoveTemp(Args));
+		FormattedTokensString = GEngine->GetEngineSubsystem<UNamingTokensEngineSubsystem>()->GetFormattedTokensStringForDisplay(MoveTemp(Args));
 	}
 	return FText::FromString(FormattedTokensString);
 }
diff --git a/Engine/Plugins/Developer/NamingTokens/NamingTokens.uplugin b/Engine/Plugins/Developer/NamingTokens/NamingTokens.uplugin
--- a/Engine/Plugins/Developer/NamingTokens/NamingTokens.uplugin
+++ b/Engine/Plugins/Developer/NamingTokens/NamingTokens.uplugin
@@ -5,8 +5,8 @@
 	"FriendlyName": "NamingTokens",
 	"Description": "Define tokens which can be recognized during string evaluation.",
 	"Category": "Editor",
-	"CreatedBy" : "Epic Games, Inc.",
-	"CreatedByURL" : "http://epicgames.com",
+	"CreatedBy": "Epic Games, Inc.",
+	"CreatedByURL": "http://epicgames.com",
 	"DocsURL": "",
 	"MarketplaceURL": "",
 	"CanContainContent": true,
@@ -16,6 +16,11 @@
 	"Modules": [
 		{
 			"Name": "NamingTokens",
+			"Type": "Runtime",
+			"LoadingPhase": "Default"
+		},
+		{
+			"Name": "NamingTokensUncookedOnly",
 			"Type": "UncookedOnly",
 			"LoadingPhase": "Default"
 		}
diff --git a/Engine/Plugins/Developer/NamingTokens/Source/NamingTokens/NamingTokens.Build.cs b/Engine/Plugins/Developer/NamingTokens/Source/NamingTokens/NamingTokens.Build.cs
--- a/Engine/Plugins/Developer/NamingTokens/Source/NamingTokens/NamingTokens.Build.cs
+++ b/Engine/Plugins/Developer/NamingTokens/Source/NamingTokens/NamingTokens.Build.cs
@@ -12,22 +12,14 @@ public class NamingTokens : ModuleRules
 			new string[]
 			{
 				"Core",
-				"EditorSubsystem"
+				"Engine"
 			}
 		);
 		
 		PrivateDependencyModuleNames.AddRange(
 			new string[]
 			{
-				"CoreUObject",
-				"Engine",
-				"Slate",
-				"SlateCore",
-				"InputCore",
-				"UnrealEd",
-				"BlueprintGraph",
-				"AssetDefinition",
-				"EngineAssetDefinitions",
+				"CoreUObject"
 			}
 		);
 	}
diff --git a/Engine/Plugins/Developer/NamingTokens/Source/NamingTokens/Private/AssetDefinition_NamingTokens.cpp b/Engine/Plugins/Developer/NamingTokens/Source/NamingTokens/Private/AssetDefinition_NamingTokens.cpp
deleted file mode 100644
diff --git a/Engine/Plugins/Developer/NamingTokens/Source/NamingTokens/Private/AssetDefinition_NamingTokens.h b/Engine/Plugins/Developer/NamingTokens/Source/NamingTokens/Private/AssetDefinition_NamingTokens.h
deleted file mode 100644
diff --git a/Engine/Plugins/Developer/NamingTokens/Source/NamingTokens/Private/Customization/NamingTokensCustomization.cpp b/Engine/Plugins/Developer/NamingTokens/Source/NamingTokens/Private/Customization/NamingTokensCustomization.cpp
deleted file mode 100644
diff --git a/Engine/Plugins/Developer/NamingTokens/Source/NamingTokens/Private/Customization/NamingTokensCustomization.h b/Engine/Plugins/Developer/NamingTokens/Source/NamingTokens/Private/Customization/NamingTokensCustomization.h
deleted file mode 100644
diff --git a/Engine/Plugins/Developer/NamingTokens/Source/NamingTokens/Private/NamingTokens.cpp b/Engine/Plugins/Developer/NamingTokens/Source/NamingTokens/Private/NamingTokens.cpp
--- a/Engine/Plugins/Developer/NamingTokens/Source/NamingTokens/Private/NamingTokens.cpp
+++ b/Engine/Plugins/Developer/NamingTokens/Source/NamingTokens/Private/NamingTokens.cpp
@@ -67,7 +67,7 @@ FNamingTokenResultData UNamingTokens::EvaluateTokenText(const FText& InTokenText
 	FText FormattedText = InTokenText;
 
 	// Retrieve all unprocessed tokens within this text.
-	const TArray<FString> UnprocessedTokens = UE::NamingTokens::Utils::Private::GetTokenKeysFromString(FormattedText.ToString());
+	const TArray<FString> UnprocessedTokens = UE::NamingTokens::Utils::GetTokenKeysFromString(FormattedText.ToString());
 
 	TSet<FNamingTokenData> DefaultTokensSet(GetDefaultTokens());
 	TSet<FNamingTokenData> CustomTokenSet(CustomTokens);
@@ -76,11 +76,11 @@ FNamingTokenResultData UNamingTokens::EvaluateTokenText(const FText& InTokenText
 	
 	for (const FString& TokenKeyWithNamespace : UnprocessedTokens)
 	{
-		const FString TokenNamespace = UE::NamingTokens::Utils::Private::GetNamespaceFromTokenKey(TokenKeyWithNamespace);
+		const FString TokenNamespace = UE::NamingTokens::Utils::GetNamespaceFromTokenKey(TokenKeyWithNamespace);
 		if (TokenNamespace.IsEmpty() || TokenNamespace == Namespace)
 		{
 			// Make sure the token key doesn't include the namespace.
-			const FString TokenKey = UE::NamingTokens::Utils::Private::RemoveNamespaceFromTokenKey(TokenKeyWithNamespace);
+			const FString TokenKey = UE::NamingTokens::Utils::RemoveNamespaceFromTokenKey(TokenKeyWithNamespace);
 
 			if (CheckedKeys.Contains(TokenKey))
 			{
@@ -144,13 +144,13 @@ FString UNamingTokens::GetFormattedTokensStringForDisplay() const
 	for (const FNamingTokenData& Token : GetDefaultTokens())
 	{
 		FormattedTokensString += FString::Printf(TEXT("%s - %s\n"),
-			*UE::NamingTokens::Utils::Private::CreateFormattedToken(Token), *Token.DisplayName.ToString());
+			*UE::NamingTokens::Utils::CreateFormattedToken(Token), *Token.DisplayName.ToString());
 	}
 
 	for (const FNamingTokenData& Token : CustomTokens)
 	{
 		FormattedTokensString += FString::Printf(TEXT("%s - %s\n"),
-			*UE::NamingTokens::Utils::Private::CreateFormattedToken(Token), *Token.DisplayName.ToString());
+			*UE::NamingTokens::Utils::CreateFormattedToken(Token), *Token.DisplayName.ToString());
 	}
 	
 	return FormattedTokensString;
@@ -229,7 +229,7 @@ void UNamingTokens::OnPostEvaluate_Implementation()
 
 FText UNamingTokens::ProcessToken(const FNamingTokenData& InToken, FText& InOutFormattedText)
 {
-	const FString FormattedToken = *UE::NamingTokens::Utils::Private::CreateFormattedToken(InToken);
+	const FString FormattedToken = *UE::NamingTokens::Utils::CreateFormattedToken(InToken);
 
 	FText ProcessedToken = FText::GetEmpty();
 
@@ -265,7 +265,7 @@ UFunction* UNamingTokens::FindBlueprintFunctionForToken(const FNamingTokenData&
 	const FName FunctionName = InTokenData.FunctionName;
 	if (UFunction* Function = GetClass()->FindFunctionByName(FunctionName))
 	{
-		if (UE::NamingTokens::Utils::Private::ValidateTokenFunction(Function))
+		if (UE::NamingTokens::Utils::ValidateTokenFunction(Function))
 		{
 			return Function;
 		}
diff --git a/Engine/Plugins/Developer/NamingTokens/Source/NamingTokens/Private/NamingTokensEditorSubsystem.cpp b/Engine/Plugins/Developer/NamingTokens/Source/NamingTokens/Private/NamingTokensEditorSubsystem.cpp
deleted file mode 100644
diff --git a/Engine/Plugins/Developer/NamingTokens/Source/NamingTokens/Private/NamingTokensEngineSubsystem.cpp b/Engine/Plugins/Developer/NamingTokens/Source/NamingTokens/Private/NamingTokensEngineSubsystem.cpp
new file mode 100644
--- /dev/null
+++ b/Engine/Plugins/Developer/NamingTokens/Source/NamingTokens/Private/NamingTokensEngineSubsystem.cpp
@@ -0,0 +1,359 @@
+// Copyright Epic Games, Inc. All Rights Reserved.
+
+#include "NamingTokensEngineSubsystem.h"
+
+#include "NamingTokens.h"
+#include "NamingTokensEvaluationData.h"
+#include "NamingTokensLog.h"
+#include "Utils/NamingTokenUtils.h"
+
+#include "AssetRegistry/AssetData.h"
+#include "AssetRegistry/AssetDataTagMap.h"
+#include "AssetRegistry/AssetRegistryModule.h"
+#include "Engine/Blueprint.h"
+#include "Internationalization/Regex.h"
+#include "Misc/PackageName.h"
+
+TArray<FString> UNamingTokensEngineSubsystem::GetAllNamespaces() const
+{
+	const FAssetRegistryModule& AssetRegistryModule = FModuleManager::LoadModuleChecked<FAssetRegistryModule>("AssetRegistry");
+	const IAssetRegistry& AssetRegistry = AssetRegistryModule.Get();
+
+	// Find all assets whose native parent class is UNamingTokens, and add their namespaces
+	TArray<FAssetData> NamingTokenAssets;
+	const TMultiMap<FName, FString> TagValues = { { FBlueprintTags::NativeParentClassPath, FObjectPropertyBase::GetExportPath(UNamingTokens::StaticClass()) } };
+	AssetRegistry.GetAssetsByTagValues(TagValues, NamingTokenAssets);
+
+	TSet<FString> DiscoveredNamespaces;
+	for (const FAssetData& NamingTokenAssetData : NamingTokenAssets)
+	{
+		FString FoundNamespace;
+		NamingTokenAssetData.GetTagValue(UNamingTokens::GetNamespacePropertyName(), FoundNamespace);
+
+		DiscoveredNamespaces.Add(FoundNamespace);
+	}
+
+	// Find all native classes derived from UNamingTokens, and add their namespaces
+	TArray<UClass*> DerivedClasses;
+	GetDerivedClasses(UNamingTokens::StaticClass(), DerivedClasses);
+	for (const UClass* DerivedClass : DerivedClasses)
+	{
+		constexpr EClassFlags InvalidClassFlags = CLASS_Deprecated | CLASS_NewerVersionExists | CLASS_Abstract;
+		if (DerivedClass->IsNative() && !DerivedClass->HasAnyClassFlags(InvalidClassFlags))
+		{
+			FString Namepsace = DerivedClass->GetDefaultObject<UNamingTokens>()->GetNamespace();
+			DiscoveredNamespaces.Add(Namepsace);
+		}
+	}
+
+	return DiscoveredNamespaces.Array();
+}
+
+UNamingTokens* UNamingTokensEngineSubsystem::GetNamingTokens(const FString& InNamespace) const
+{
+	auto LoadNamingToken = [this, &InNamespace](const TSoftClassPtr<UNamingTokens>& InTokensClass) -> UNamingTokens*
+	{
+		if (const UClass* Class = InTokensClass.LoadSynchronous())
+		{
+			UNamingTokens* NamingTokens = NewObject<UNamingTokens>(GetTransientPackage(), Class);
+			if (ensure(NamingTokens->GetNamespace() == InNamespace))
+			{
+				CachedNamingTokens.Add(InNamespace, NamingTokens);
+				return NamingTokens;
+			}
+		}
+		return nullptr;
+	};
+
+	// Check cache.
+	if (const TObjectPtr<UNamingTokens>* Tokens = CachedNamingTokens.Find(InNamespace))
+	{
+		return *Tokens;
+	}
+
+	UNamingTokens* FoundTokens = nullptr;
+
+	const UClass* TargetClass = UNamingTokens::StaticClass();
+	
+	// Check native classes.
+	{
+		constexpr EClassFlags InvalidClassFlags = CLASS_Deprecated | CLASS_NewerVersionExists | CLASS_Abstract;
+
+		TArray<UClass*> DerivedClasses;
+		GetDerivedClasses(TargetClass, DerivedClasses);
+		for (const UClass* DerivedClass : DerivedClasses)
+		{
+			if (DerivedClass->IsNative()
+				&& DerivedClass->IsChildOf(TargetClass)
+				&& !DerivedClass->HasAnyClassFlags(InvalidClassFlags)
+				&& !DerivedClass->GetName().StartsWith(TEXT("SKEL_"))
+				&& !DerivedClass->GetName().StartsWith(TEXT("REINST_"))
+				&& DerivedClass->GetDefaultObject<UNamingTokens>()->GetNamespace() == InNamespace)
+			{
+				if (!FoundTokens)
+				{
+					FoundTokens = LoadNamingToken(DerivedClass);
+				}
+				else
+				{
+					UE_LOG(LogNamingTokens, Warning, TEXT("Namespace '%s' exists more than once in native class '%s' and '%s'."),
+						*InNamespace, *FoundTokens->GetClass()->GetName(), *DerivedClass->GetName());
+				}
+			}
+		}
+	}
+
+	// Check blueprint classes.
+	// NOTE: We may want to use OnAssetAdded and cache the asset or load the class then.  
+	{
+		const FAssetRegistryModule& AssetRegistryModule = FModuleManager::LoadModuleChecked<FAssetRegistryModule>("AssetRegistry");
+		const IAssetRegistry& AssetRegistry = AssetRegistryModule.Get();
+
+		FARFilter Filter;
+		Filter.ClassPaths.Add(UBlueprint::StaticClass()->GetClassPathName());
+		Filter.bRecursiveClasses = true;
+		Filter.bRecursivePaths = true;
+
+		TArray<FAssetData> BlueprintAssets;
+		AssetRegistry.GetAssets(Filter, BlueprintAssets);
+
+		TSet<FTopLevelAssetPath> DerivedClassNames;
+		AssetRegistry.GetDerivedClassNames({ TargetClass->GetClassPathName() }, {}, DerivedClassNames);
+
+		for (const FAssetData& AssetData : BlueprintAssets)
+		{
+			// Narrow down to only our assets.
+			const FAssetDataTagMapSharedView::FFindTagResult Result = AssetData.TagsAndValues.FindTag(TEXT("GeneratedClass"));
+			if (Result.IsSet())
+			{
+				const FString& GeneratedClassPathPtr = Result.GetValue();
+				const FTopLevelAssetPath ClassObjectPath(FPackageName::ExportTextPathToObjectPath(*GeneratedClassPathPtr));
+
+				if (DerivedClassNames.Contains(ClassObjectPath))
+				{
+					// Now check for the namespace match.
+					FString FoundNamespace;
+					if (AssetData.GetTagValue(UNamingTokens::GetNamespacePropertyName(), FoundNamespace)
+						&& FoundNamespace == InNamespace)
+					{
+						// Allow BP children to overwrite native classes.
+						auto IsAssetChildOfNativeToken = [&AssetRegistry, &InNamespace](const FTopLevelAssetPath& InAssetPath, const UNamingTokens* InParentTokens) -> bool
+						{
+							check(InParentTokens);
+							if (InParentTokens->GetClass()->IsNative())
+							{
+								// We need to make sure this BP is a child of the chosen native class.
+								TSet<FTopLevelAssetPath> NativeDerivedClassNames;
+								AssetRegistry.GetDerivedClassNames({ InParentTokens->GetClass()->GetClassPathName() },
+									{}, NativeDerivedClassNames);
+
+								if (NativeDerivedClassNames.Contains(InAssetPath))
+								{
+									UE_LOG(LogNamingTokens, Log, TEXT("Using namespace '%s' of BP child '%s' instead of native parent '%s'."),
+										*InNamespace, *InAssetPath.ToString(), *InParentTokens->GetClass()->GetName());
+									return true;
+								}
+							}
+							return false;
+						};
+						
+						if (!FoundTokens || IsAssetChildOfNativeToken(ClassObjectPath, FoundTokens))
+						{
+							const TSoftClassPtr<UNamingTokens> SoftClassPath(ClassObjectPath.ToString());
+							FoundTokens = LoadNamingToken(SoftClassPath);
+						}
+						else
+						{
+							UE_LOG(LogNamingTokens, Warning, TEXT("Namespace '%s' exists more than once in class '%s' and BP asset '%s'."),
+								*InNamespace, *FoundTokens->GetClass()->GetName(), *AssetData.PackageName.ToString());
+						}
+					}
+				}
+			}
+		}
+	}
+
+	return FoundTokens;
+}
+
+TArray<UNamingTokens*> UNamingTokensEngineSubsystem::GetMultipleNamingTokens(const TArray<FString>& InNamespaces) const
+{
+	TArray<UNamingTokens*> Result;
+	const TSet<FString> UniqueNamespaces(InNamespaces);
+	for (const FString& Namespace : UniqueNamespaces)
+	{
+		if (UNamingTokens* NamingTokens = GetNamingTokens(Namespace))
+		{
+			Result.Add(NamingTokens);
+		}
+	}
+
+	return Result;
+}
+
+FNamingTokenResultData UNamingTokensEngineSubsystem::EvaluateTokenText(const FText& InTokenText, const FNamingTokenFilterArgs& InFilter, const TArray<UObject*>& InContexts)
+{
+	FNamingTokenResultData Result;
+
+	FText ProcessedTokenText = InTokenText;
+	Result.EvaluatedText = ProcessedTokenText; // Set so we always have something in case no tokens were evaluated.
+
+	TSet<FString> Namespaces = GetNamingTokenNamespacesFromString(InTokenText.ToString(), InFilter);
+
+	TArray<FNamingTokenValueData> CompletedNamingTokenValueData;
+	
+	// These external filters are intentionally executed before adding any namespaces from the input FNamingTokenFilterArgs filter. 
+	// This allows the caller of this function to guarantee that the namespaces they want included will always be allowed, regardless of other filters that may execute.
+	for (const TPair<FName, FFilterNamespace>& DelegatePair : FilterNamespaceDelegates)
+	{
+		DelegatePair.Value.Execute(Namespaces);
+	}
+
+	Namespaces.Append(InFilter.NamespacesToInclude); // Tokens could have been written without the namespace but should still be scoped to the filter.
+
+	FNamingTokensEvaluationData EvaluationData;
+	EvaluationData.Contexts = InContexts;
+	for (const FString& Namespace : Namespaces)
+	{
+		if (UNamingTokens* NamingTokens = GetNamingTokens(Namespace))
+		{
+			Result = NamingTokens->EvaluateTokenText(ProcessedTokenText, EvaluationData);
+
+			// Properly record individual token evaluation data. The order should be consistent with the order
+			// a key was processed in the string. A key may be processed multiple times if it is referenced multiple times, or is not
+			// identified in one NamingTokens class but is in another. In the event a future class identifies it, it will update the original
+			// key.
+			for (const FNamingTokenValueData& NewTokenValue : Result.TokenValues)
+			{
+				bool bReplacedOldValue = false;
+				// Find previously recorded values that were undefined and update them indicating they are no longer undefined.
+				for (FNamingTokenValueData& OldTokenValue : CompletedNamingTokenValueData)
+				{
+					if (OldTokenValue.TokenKey.Equals(NewTokenValue.TokenKey, ESearchCase::CaseSensitive))
+					{
+						OldTokenValue = NewTokenValue;
+						bReplacedOldValue = true;
+						break;
+					}
+				}
+				if (!bReplacedOldValue)
+				{
+					// Add to the completed naming token value.
+					CompletedNamingTokenValueData.Add(NewTokenValue);
+				}
+			}
+			
+			ProcessedTokenText = Result.EvaluatedText;
+		}
+	}
+
+	Result.OriginalText = InTokenText;
+
+	Result.TokenValues = MoveTemp(CompletedNamingTokenValueData);
+	
+	return Result;
+}
+
+FNamingTokenResultData UNamingTokensEngineSubsystem::EvaluateTokenText(const FText& InTokenText, const FNamingTokenFilterArgs& InFilter)
+{
+	return EvaluateTokenText(InTokenText, InFilter, {});
+}
+
+FNamingTokenResultData UNamingTokensEngineSubsystem::EvaluateTokenString(const FString& InTokenString, const FNamingTokenFilterArgs& InFilter, const TArray<UObject*>& InContexts)
+{
+	return EvaluateTokenText(FText::FromString(InTokenString), InFilter, InContexts);
+}
+
+FNamingTokenResultData UNamingTokensEngineSubsystem::EvaluateTokenString(const FString& InTokenString, const FNamingTokenFilterArgs& InFilter)
+{
+	return EvaluateTokenString(InTokenString, InFilter, {});
+}
+
+void UNamingTokensEngineSubsystem::RegisterGlobalNamespace(const FString& InNamespace)
+{
+	if (IsGlobalNamespaceRegistered(InNamespace))
+	{
+		const FString UnregisterFunctionName = GET_FUNCTION_NAME_CHECKED(UNamingTokensEngineSubsystem, UnregisterGlobalNamespace).ToString();
+		UE_LOG(LogNamingTokens, Error, TEXT("NamingTokens Namespace '%s' is already registered as a global namespace. Call '%s' first if you wish to overwrite the namespace."),
+			 *InNamespace, *UnregisterFunctionName);
+		return;
+	}
+	GlobalNamespaces.Add(InNamespace);
+}
+
+void UNamingTokensEngineSubsystem::UnregisterGlobalNamespace(const FString& InNamespace)
+{
+	GlobalNamespaces.Remove(InNamespace);
+}
+
+bool UNamingTokensEngineSubsystem::IsGlobalNamespaceRegistered(const FString& InNamespace) const
+{
+	return GlobalNamespaces.Contains(InNamespace);
+}
+
+TArray<FString> UNamingTokensEngineSubsystem::GetGlobalNamespaces() const
+{
+	return GlobalNamespaces.Array();
+}
+
+void UNamingTokensEngineSubsystem::RegisterNamespaceFilter(const FName OwnerName, FFilterNamespace Delegate)
+{
+	FilterNamespaceDelegates.Add(OwnerName, MoveTemp(Delegate));
+}
+
+void UNamingTokensEngineSubsystem::UnregisterNamespaceFilter(const FName OwnerName)
+{
+	FilterNamespaceDelegates.Remove(OwnerName);
+}
+
+TSet<FString> UNamingTokensEngineSubsystem::GetNamingTokenNamespacesFromString(const FString& InTokenString, const FNamingTokenFilterArgs& InFilter) const
+{
+	TSet<FString> Namespaces;
+
+	const FRegexPattern Pattern(R"(\{\s*([a-zA-Z0-9_]+)\.[a-zA-Z0-9_]+\s*\})");
+	FRegexMatcher Matcher(Pattern, InTokenString);
+	
+	while (Matcher.FindNext())
+	{
+		const FString Namespace = Matcher.GetCaptureGroup(1);
+		if (InFilter.NamespacesToInclude.IsEmpty() || InFilter.NamespacesToInclude.Contains(Namespace))
+		{
+			Namespaces.Add(Namespace);
+		}
+	}
+
+	// Add any global namespaces.
+	if (InFilter.bIncludeGlobal)
+	{
+		Namespaces.Append(GlobalNamespaces);
+	}
+	
+	return Namespaces;	
+}
+
+FString UNamingTokensEngineSubsystem::GetFormattedTokensStringForDisplay(const FNamingTokenFilterArgs& InFilter) const
+{
+	TSet<FString> Namespaces = TSet<FString>(InFilter.NamespacesToInclude);
+
+	if (InFilter.bIncludeGlobal)
+	{
+		Namespaces.Append(GlobalNamespaces.Array());
+	}
+
+	FString Result;
+
+	for (const FString& Namespace : Namespaces)
+	{
+		if (const UNamingTokens* NamingTokens = GetNamingTokens(Namespace))
+		{
+			Result += NamingTokens->GetFormattedTokensStringForDisplay();
+		}
+	}
+
+	return Result;
+}
+
+void UNamingTokensEngineSubsystem::ClearCachedNamingTokens()
+{
+	CachedNamingTokens.Empty();
+}
diff --git a/Engine/Plugins/Developer/NamingTokens/Source/NamingTokens/Private/NamingTokensFactory.cpp b/Engine/Plugins/Developer/NamingTokens/Source/NamingTokens/Private/NamingTokensFactory.cpp
deleted file mode 100644
diff --git a/Engine/Plugins/Developer/NamingTokens/Source/NamingTokens/Private/NamingTokensFactory.h b/Engine/Plugins/Developer/NamingTokens/Source/NamingTokens/Private/NamingTokensFactory.h
deleted file mode 100644
diff --git a/Engine/Plugins/Developer/NamingTokens/Source/NamingTokens/Private/NamingTokensModule.cpp b/Engine/Plugins/Developer/NamingTokens/Source/NamingTokens/Private/NamingTokensModule.cpp
--- a/Engine/Plugins/Developer/NamingTokens/Source/NamingTokens/Private/NamingTokensModule.cpp
+++ b/Engine/Plugins/Developer/NamingTokens/Source/NamingTokens/Private/NamingTokensModule.cpp
@@ -2,38 +2,24 @@
 
 #include "NamingTokensModule.h"
 
-#include "Customization/NamingTokensCustomization.h"
 #include "GlobalNamingTokens.h"
-#include "NamingTokenData.h"
-#include "NamingTokensEditorSubsystem.h"
+#include "NamingTokensEngineSubsystem.h"
 
-#include "Editor.h"
-#include "PropertyEditorModule.h"
+#include "Engine/Engine.h"
+#include "Misc/CoreDelegates.h"
 
 #define LOCTEXT_NAMESPACE "FNamingTokensModule"
 
 void FNamingTokensModule::StartupModule()
 {
-	FPropertyEditorModule& PropertyModule = FModuleManager::GetModuleChecked<FPropertyEditorModule>("PropertyEditor");
-	PropertyModule.RegisterCustomPropertyTypeLayout(FNamingTokenData::StaticStruct()->GetFName(),
-		FOnGetPropertyTypeCustomizationInstance::CreateStatic(&FNamingTokensCustomization::MakeInstance));
-
 	FCoreDelegates::OnPostEngineInit.AddRaw(this, &FNamingTokensModule::OnPostEngineInit);
 }
 
 void FNamingTokensModule::ShutdownModule()
 {
-	if (UObjectInitialized() && !IsEngineExitRequested())
-	{
-		if (FPropertyEditorModule* PropertyEditorModule = FModuleManager::GetModulePtr<FPropertyEditorModule>("PropertyEditor"))
-		{
-			PropertyEditorModule->UnregisterCustomPropertyTypeLayout(FNamingTokenData::StaticStruct()->GetFName());
-		}
-	}
-	
-	if (GEditor)
+	if (GEngine)
 	{
-		if (UNamingTokensEditorSubsystem* NamingTokensSubsystem = GEditor->GetEditorSubsystem<UNamingTokensEditorSubsystem>())
+		if (UNamingTokensEngineSubsystem* NamingTokensSubsystem = GEngine->GetEngineSubsystem<UNamingTokensEngineSubsystem>())
 		{
 			NamingTokensSubsystem->UnregisterGlobalNamespace(UGlobalNamingTokens::GetGlobalNamespace());
 		}
@@ -44,9 +30,9 @@ void FNamingTokensModule::ShutdownModule()
 
 void FNamingTokensModule::OnPostEngineInit()
 {
-	if (GEditor)
+	if (GEngine)
 	{
-		if (UNamingTokensEditorSubsystem* NamingTokensSubsystem = GEditor->GetEditorSubsystem<UNamingTokensEditorSubsystem>())
+		if (UNamingTokensEngineSubsystem* NamingTokensSubsystem = GEngine->GetEngineSubsystem<UNamingTokensEngineSubsystem>())
 		{
 			if (!NamingTokensSubsystem->IsGlobalNamespaceRegistered(UGlobalNamingTokens::GetGlobalNamespace()))
 			{
diff --git a/Engine/Plugins/Developer/NamingTokens/Source/NamingTokens/Private/Utils/NamingTokenUtils.cpp b/Engine/Plugins/Developer/NamingTokens/Source/NamingTokens/Private/Utils/NamingTokenUtils.cpp
--- a/Engine/Plugins/Developer/NamingTokens/Source/NamingTokens/Private/Utils/NamingTokenUtils.cpp
+++ b/Engine/Plugins/Developer/NamingTokens/Source/NamingTokens/Private/Utils/NamingTokenUtils.cpp
@@ -1,22 +1,20 @@
 // Copyright Epic Games, Inc. All Rights Reserved.
 
-#include "NamingTokenUtils.h"
+#include "Utils/NamingTokenUtils.h"
 
 #include "NamingTokens.h"
 
 #include "Internationalization/LocKeyFuncs.h"
 #include "Internationalization/Regex.h"
-#include "Kismet2/BlueprintEditorUtils.h"
-#include "Kismet2/KismetEditorUtilities.h"
 
-UFunction* UE::NamingTokens::Utils::Private::GetProcessTokenFunctionSignature()
+UFunction* UE::NamingTokens::Utils::GetProcessTokenFunctionSignature()
 {
 	UFunction* FunctionSignature = UNamingTokens::StaticClass()->FindFunctionByName(UNamingTokens::GetProcessTokenTemplateFunctionName());
 	check(FunctionSignature);
 	return FunctionSignature;
 }
 
-bool UE::NamingTokens::Utils::Private::ValidateTokenFunction(const UFunction* InFunction)
+bool UE::NamingTokens::Utils::ValidateTokenFunction(const UFunction* InFunction)
 {
 	check(InFunction);
 	const UFunction* FunctionSignature = GetProcessTokenFunctionSignature();
@@ -27,36 +25,12 @@ bool UE::NamingTokens::Utils::Private::ValidateTokenFunction(const UFunction* In
 				&& !InFunction->HasAnyFunctionFlags(FUNC_Private);
 }
 
-FString UE::NamingTokens::Utils::Private::CreateBaseTokenFunctionName(const FString& InTokenKey)
-{
-	return TEXT("ProcessToken_") + InTokenKey;
-}
-
-FName UE::NamingTokens::Utils::Private::CreateNewTokenGraph(UBlueprint* InBlueprint, const FString& InTokenKey)
-{
-	check(InBlueprint);
-	
-	const FString BaseName = CreateBaseTokenFunctionName(InTokenKey);
-	const FName FunctionName = FBlueprintEditorUtils::GenerateUniqueGraphName(InBlueprint, BaseName);
-
-	InBlueprint->Modify();
-		
-	UEdGraph* NewGraph = FBlueprintEditorUtils::CreateNewGraph(InBlueprint, FunctionName, UEdGraph::StaticClass(), UEdGraphSchema_K2::StaticClass());
-	
-	UFunction* FunctionSignature = GetProcessTokenFunctionSignature();
-		
-	FBlueprintEditorUtils::AddFunctionGraph<UFunction>(InBlueprint, NewGraph, /*bIsUserCreated=*/ true, FunctionSignature);
-	FKismetEditorUtilities::BringKismetToFocusAttentionOnObject(NewGraph);
-
-	return FunctionName;
-}
-
-FString UE::NamingTokens::Utils::Private::CreateFormattedToken(const FNamingTokenData& InToken)
+FString UE::NamingTokens::Utils::CreateFormattedToken(const FNamingTokenData& InToken)
 {
 	return FString::Printf(TEXT("{%s}"), *InToken.TokenKey);
 }
 
-TArray<FString> UE::NamingTokens::Utils::Private::GetTokenKeysFromString(const FString& InTokenString)
+TArray<FString> UE::NamingTokens::Utils::GetTokenKeysFromString(const FString& InTokenString)
 {
 	TSet<FString, FLocKeySetFuncs> Tokens;
 
@@ -73,7 +47,7 @@ TArray<FString> UE::NamingTokens::Utils::Private::GetTokenKeysFromString(const F
 	return Tokens.Array();
 }
 
-FString UE::NamingTokens::Utils::Private::GetNamespaceFromTokenKey(const FString& InTokenKey)
+FString UE::NamingTokens::Utils::GetNamespaceFromTokenKey(const FString& InTokenKey)
 {
 	FString LeftPart, RightPart;
 	InTokenKey.Split(TEXT("."), &LeftPart, &RightPart);
@@ -81,7 +55,7 @@ FString UE::NamingTokens::Utils::Private::GetNamespaceFromTokenKey(const FString
 	return LeftPart;
 }
 
-FString UE::NamingTokens::Utils::Private::RemoveNamespaceFromTokenKey(const FString& InTokenKey)
+FString UE::NamingTokens::Utils::RemoveNamespaceFromTokenKey(const FString& InTokenKey)
 {
 	FString Result = InTokenKey;
 	FString LeftPart, RightPart;
diff --git a/Engine/Plugins/Developer/NamingTokens/Source/NamingTokens/Private/Utils/NamingTokenUtils.h b/Engine/Plugins/Developer/NamingTokens/Source/NamingTokens/Private/Utils/NamingTokenUtils.h
deleted file mode 100644
diff --git a/Engine/Plugins/Developer/NamingTokens/Source/NamingTokens/Public/NamingTokens.h b/Engine/Plugins/Developer/NamingTokens/Source/NamingTokens/Public/NamingTokens.h
--- a/Engine/Plugins/Developer/NamingTokens/Source/NamingTokens/Public/NamingTokens.h
+++ b/Engine/Plugins/Developer/NamingTokens/Source/NamingTokens/Public/NamingTokens.h
@@ -102,15 +102,13 @@ protected:
 	UFUNCTION(BlueprintCallable, BlueprintNativeEvent, Category = "NamingTokens")
 	FDateTime GetCurrentDateTime() const;
 	virtual FDateTime GetCurrentDateTime_Implementation() const;
-	
-#if WITH_EDITOR
+
 public:
 	static FName GetProcessTokenTemplateFunctionName() { return GET_FUNCTION_NAME_CHECKED(UNamingTokens, ProcessTokenTemplateFunction); } 
 private:
 	/** Template function for us to dynamically create subclass graphs from matching this signature. */
 	UFUNCTION()
 	FText ProcessTokenTemplateFunction() { return FText::GetEmpty(); }
-#endif
 
 protected:
 	/**
diff --git a/Engine/Plugins/Developer/NamingTokens/Source/NamingTokens/Public/NamingTokensEditorSubsystem.h b/Engine/Plugins/Developer/NamingTokens/Source/NamingTokens/Public/NamingTokensEditorSubsystem.h
deleted file mode 100644
diff --git a/Engine/Plugins/Developer/NamingTokens/Source/NamingTokens/Public/NamingTokensEngineSubsystem.h b/Engine/Plugins/Developer/NamingTokens/Source/NamingTokens/Public/NamingTokensEngineSubsystem.h
new file mode 100644
--- /dev/null
+++ b/Engine/Plugins/Developer/NamingTokens/Source/NamingTokens/Public/NamingTokensEngineSubsystem.h
@@ -0,0 +1,141 @@
+// Copyright Epic Games, Inc. All Rights Reserved.
+
+#pragma once
+
+#include "Subsystems/EngineSubsystem.h"
+
+#include "NamingTokenData.h"
+
+#include "NamingTokensEngineSubsystem.generated.h"
+
+class UNamingTokens;
+
+DECLARE_DELEGATE_OneParam(FFilterNamespace, TSet<FString>& /*Namespaces*/);
+
+USTRUCT(BlueprintType)
+struct FNamingTokenFilterArgs
+{
+	GENERATED_BODY()
+	
+	/**
+	 * Namespaces to always be included during evaluation. Namespaces added here won't require the 'namespace.' string prefixed to tokens.
+	 * This does not filter out any namespaces.
+	 */
+	UPROPERTY(EditAnywhere, BlueprintReadWrite, Category=NamingTokens)
+	TArray<FString> NamespacesToInclude;
+	/** Include global namespaces. */
+	UPROPERTY(EditAnywhere, BlueprintReadWrite, Category=NamingTokens)
+	bool bIncludeGlobal = true;
+};
+
+/**
+ * An editor subsystem for registering global tokens and evaluating strings across the entire project.
+ */
+UCLASS()
+class NAMINGTOKENS_API UNamingTokensEngineSubsystem : public UEngineSubsystem
+{
+	GENERATED_BODY()
+
+public:
+	/**
+	 * Lookup naming tokens given a namespace. This will look first in cached naming tokens,
+	 * then native classes, then blueprint classes. Call ClearCachedNamingTokens to reset the cache
+	 * and force a full lookup.
+	 *
+	 * @param InNamespace The namespace of the tokens.
+	 * @return The found Naming Tokens object, or nullptr.
+	 */
+	UFUNCTION(BlueprintCallable, Category = "NamingTokens")
+	UNamingTokens* GetNamingTokens(const FString& InNamespace) const;
+
+	/**
+	 * Lookup multiple naming tokens from multiple namespaces.
+	 *
+	 * @param InNamespaces An array of all token namespaces.
+	 * @return An array of matching tokens objects.
+	 */
+	UFUNCTION(BlueprintCallable, Category = "NamingTokens")
+	TArray<UNamingTokens*> GetMultipleNamingTokens(const TArray<FString>& InNamespaces) const;
+
+	/**
+	 * Parse and evaluate token text.
+	 * 
+	 * @param InTokenText The text containing unprocessed tokens.
+	 * @param InFilter [Optional] Filter to determine which namespaces to use.
+	 * @param InContexts [Optional] Context objects to pass to naming tokens.
+	 *
+	 * @return The result of the evaluation.
+	 */
+	UFUNCTION(BlueprintCallable, Category = "NamingTokens", meta = (AutoCreateRefTerm = "InFilter,InContexts"))
+	FNamingTokenResultData EvaluateTokenText(const FText& InTokenText, const FNamingTokenFilterArgs& InFilter, const TArray<UObject*>& InContexts);
+	FNamingTokenResultData EvaluateTokenText(const FText& InTokenText, const FNamingTokenFilterArgs& InFilter = FNamingTokenFilterArgs());
+	
+	/**
+	 * Parse and evaluate token string.
+	 * 
+	 * @param InTokenString The string containing unprocessed tokens.
+	 * @param InFilter [Optional] Filter to determine which namespaces to use.
+	 * @param InContexts [Optional] Context objects to pass to naming tokens.
+	 *
+	 * @return The result of the evaluation.
+	 */
+	UFUNCTION(BlueprintCallable, Category = "NamingTokens", meta = (AutoCreateRefTerm = "InFilter,InContexts"))
+	FNamingTokenResultData EvaluateTokenString(const FString& InTokenString, const FNamingTokenFilterArgs& InFilter, const TArray<UObject*>& InContexts);
+	FNamingTokenResultData EvaluateTokenString(const FString& InTokenString, const FNamingTokenFilterArgs& InFilter = FNamingTokenFilterArgs());
+	
+	/**
+	 * Register tokens as a global namespace. This prevents the need to include the namespace in a token string.
+	 * @param InNamespace The namespace of the tokens to register.
+	 */
+	UFUNCTION(BlueprintCallable, Category = "NamingTokens")
+	void RegisterGlobalNamespace(const FString& InNamespace);
+
+	/**
+	 * Remove tokens from a global namespace.
+	 * @param InNamespace The namespace of the tokens to unregister.
+	 */
+	UFUNCTION(BlueprintCallable, Category = "NamingTokens")
+	void UnregisterGlobalNamespace(const FString& InNamespace);
+
+	/** Checks if a namespace is registered globally. */
+	UFUNCTION(BlueprintCallable, Category = "NamingTokens")
+	bool IsGlobalNamespaceRegistered(const FString& InNamespace) const;
+	
+	/** Retrieve the registered global namespaces. */
+	UFUNCTION(BlueprintCallable, Category = "NamingTokens")
+	TArray<FString> GetGlobalNamespaces() const;
+
+	/** Retrieve all discovered namespaces. */
+	UFUNCTION(BlueprintCallable, Category = "NamingTokens")
+	TArray<FString> GetAllNamespaces() const;
+
+	/** Register a delegate to filter the list of namespaces that can be used to evaluate tokens */
+	void RegisterNamespaceFilter(const FName OwnerName, FFilterNamespace Delegate);
+
+	/** Remove a delegate from the list of namespace filters */
+	void UnregisterNamespaceFilter(const FName OwnerName);
+
+	/** Locate all reference naming token namespaces from a given string. */
+	TSet<FString> GetNamingTokenNamespacesFromString(const FString& InTokenString, const FNamingTokenFilterArgs& InFilter = FNamingTokenFilterArgs()) const;
+	
+	/** Creates a friendly display string of all tokens. */
+	FString GetFormattedTokensStringForDisplay(const FNamingTokenFilterArgs& InFilter) const;
+
+	/**
+	 * Clears all currently cached naming tokens. They will be loaded on demand when required.
+	 * This can avoid an editor restart if a namespace is adjusted on a BP asset, and you don't want the old namespace to access it.
+	 */
+	UFUNCTION(BlueprintCallable, Category = "NamingTokens")
+	void ClearCachedNamingTokens();
+	
+private:
+	/** Naming tokens currently loaded from assets. */
+	UPROPERTY(Transient)
+	mutable TMap<FString, TObjectPtr<UNamingTokens>> CachedNamingTokens;
+
+	/** Namespaces considered global (don't need to include namespace to access). */
+	TSet<FString> GlobalNamespaces;
+
+	/** Filters that will execute just before evaluating a token string */
+	TMap<FName, FFilterNamespace> FilterNamespaceDelegates;
+};
diff --git a/Engine/Plugins/Developer/NamingTokens/Source/NamingTokens/Public/Utils/NamingTokenUtils.h b/Engine/Plugins/Developer/NamingTokens/Source/NamingTokens/Public/Utils/NamingTokenUtils.h
new file mode 100644
--- /dev/null
+++ b/Engine/Plugins/Developer/NamingTokens/Source/NamingTokens/Public/Utils/NamingTokenUtils.h
@@ -0,0 +1,36 @@
+// Copyright Epic Games, Inc. All Rights Reserved.
+
+#pragma once
+
+#include "CoreMinimal.h"
+
+struct FNamingTokenData;
+
+namespace UE::NamingTokens::Utils
+{
+	/** Return the signature function all blueprint token processing should use. */
+	NAMINGTOKENS_API UFunction* GetProcessTokenFunctionSignature();
+
+	/** Validate that a token function can be used for token processing. */
+	NAMINGTOKENS_API bool ValidateTokenFunction(const UFunction* InFunction);
+
+	/** Formats a token key into {key}. */
+	FString CreateFormattedToken(const FNamingTokenData& InToken);
+
+	/** Extract all token keys from a string. */
+	TArray<FString> GetTokenKeysFromString(const FString& InTokenString);
+
+	/**
+	 * Return the namespace from a token key, or an empty string.
+	 * @param InTokenKey The token key shouldn't have any additional formatting beyond the namespace, eg "namespace.token"
+	 * @return The namespace, or an empty string.
+	 */
+	FString GetNamespaceFromTokenKey(const FString& InTokenKey);
+	
+	/**
+	 * Remove the namespace from the given token key.
+	 * @param InTokenKey The token key shouldn't have any additional formatting beyond the namespace, eg "namespace.token"
+	 * @return The token key without the namespace, or the InTokenKey if there was no namespace.
+	 */
+	FString RemoveNamespaceFromTokenKey(const FString& InTokenKey);
+}
\ No newline at end of file
diff --git a/Engine/Plugins/Developer/NamingTokens/Source/NamingTokensUncookedOnly/NamingTokensUncookedOnly.Build.cs b/Engine/Plugins/Developer/NamingTokens/Source/NamingTokensUncookedOnly/NamingTokensUncookedOnly.Build.cs
new file mode 100644
--- /dev/null
+++ b/Engine/Plugins/Developer/NamingTokens/Source/NamingTokensUncookedOnly/NamingTokensUncookedOnly.Build.cs
@@ -0,0 +1,32 @@
+﻿using UnrealBuildTool;
+
+public class NamingTokensUncookedOnly : ModuleRules
+{
+    public NamingTokensUncookedOnly(ReadOnlyTargetRules Target) : base(Target)
+    {
+        PCHUsage = ModuleRules.PCHUsageMode.UseExplicitOrSharedPCHs;
+
+        PublicDependencyModuleNames.AddRange(
+            new string[]
+            {
+                "Core",
+                "NamingTokens"
+            }
+        );
+
+        PrivateDependencyModuleNames.AddRange(
+            new string[]
+            {
+	            "AssetDefinition",
+	            "BlueprintGraph",
+                "CoreUObject",
+                "Engine",
+                "EngineAssetDefinitions",
+                "InputCore",
+                "Slate",
+                "SlateCore",
+                "UnrealEd",
+            }
+        );
+    }
+}
\ No newline at end of file
diff --git a/Engine/Plugins/Developer/NamingTokens/Source/NamingTokensUncookedOnly/Private/AssetDefinition_NamingTokens.cpp b/Engine/Plugins/Developer/NamingTokens/Source/NamingTokensUncookedOnly/Private/AssetDefinition_NamingTokens.cpp
new file mode 100644
--- /dev/null
+++ b/Engine/Plugins/Developer/NamingTokens/Source/NamingTokensUncookedOnly/Private/AssetDefinition_NamingTokens.cpp
@@ -0,0 +1,21 @@
+// Copyright Epic Games, Inc. All Rights Reserved.
+
+#include "AssetDefinition_NamingTokens.h"
+
+#include "NamingTokens.h"
+
+TSoftClassPtr<UObject> UAssetDefinition_NamingTokens::GetAssetClass() const
+{
+	return UNamingTokens::StaticClass();
+}
+
+TConstArrayView<FAssetCategoryPath> UAssetDefinition_NamingTokens::GetAssetCategories() const
+{
+	static const FAssetCategoryPath Categories[] = { FAssetCategoryPath(EAssetCategoryPaths::Misc) };
+	return Categories;
+}
+
+FText UAssetDefinition_NamingTokens::GetAssetDisplayName() const
+{
+	return NSLOCTEXT("AssetDefinition_NamingTokens", "AssetDefinition_NamingTokens_DisplayName", "Naming Tokens");
+}
diff --git a/Engine/Plugins/Developer/NamingTokens/Source/NamingTokensUncookedOnly/Private/AssetDefinition_NamingTokens.h b/Engine/Plugins/Developer/NamingTokens/Source/NamingTokensUncookedOnly/Private/AssetDefinition_NamingTokens.h
new file mode 100644
--- /dev/null
+++ b/Engine/Plugins/Developer/NamingTokens/Source/NamingTokensUncookedOnly/Private/AssetDefinition_NamingTokens.h
@@ -0,0 +1,18 @@
+// Copyright Epic Games, Inc. All Rights Reserved.
+
+#pragma once
+
+#include "Script/AssetDefinition_Blueprint.h"
+
+#include "AssetDefinition_NamingTokens.generated.h"
+
+UCLASS()
+class UAssetDefinition_NamingTokens : public UAssetDefinition_Blueprint
+{
+	GENERATED_BODY()
+	
+public:
+	virtual TSoftClassPtr<UObject> GetAssetClass() const override;
+	virtual TConstArrayView<FAssetCategoryPath> GetAssetCategories() const override;
+	virtual FText GetAssetDisplayName() const override;
+};
diff --git a/Engine/Plugins/Developer/NamingTokens/Source/NamingTokensUncookedOnly/Private/Customization/NamingTokensCustomization.cpp b/Engine/Plugins/Developer/NamingTokens/Source/NamingTokensUncookedOnly/Private/Customization/NamingTokensCustomization.cpp
new file mode 100644
--- /dev/null
+++ b/Engine/Plugins/Developer/NamingTokens/Source/NamingTokensUncookedOnly/Private/Customization/NamingTokensCustomization.cpp
@@ -0,0 +1,234 @@
+// Copyright Epic Games, Inc. All Rights Reserved.
+
+#include "NamingTokensCustomization.h"
+
+#include "NamingTokenData.h"
+#include "NamingTokens.h"
+#include "Utils/NamingTokenUtils.h"
+#include "Utils/NamingTokensEditorUtils.h"
+
+#include "DetailWidgetRow.h"
+#include "Engine/Blueprint.h"
+#include "IDetailChildrenBuilder.h"
+#include "IDetailPropertyRow.h"
+#include "IPropertyUtilities.h"
+#include "Widgets/Input/SButton.h"
+#include "Widgets/Input/SComboBox.h"
+#include "Widgets/Text/STextBlock.h"
+
+#define LOCTEXT_NAMESPACE "NamingTokensCustomization"
+
+void FNamingTokensCustomization::CustomizeHeader(TSharedRef<IPropertyHandle> PropertyHandle, FDetailWidgetRow& HeaderRow,
+                                                 IPropertyTypeCustomizationUtils& CustomizationUtils)
+{
+	CustomizationUtilsPtr = &CustomizationUtils;
+
+	if (const UNamingTokens* OwningTokens = GetOwningNamingTokens())
+	{
+		OwningBlueprint = UBlueprint::GetBlueprintFromClass(OwningTokens->GetClass());
+		if (OwningBlueprint.IsValid())
+		{
+			FunctionNameHandle = PropertyHandle->GetChildHandle(GET_MEMBER_NAME_CHECKED(FNamingTokenData, FunctionName));
+			check(FunctionNameHandle.IsValid());
+
+			TokenKeyHandle = PropertyHandle->GetChildHandle(GET_MEMBER_NAME_CHECKED(FNamingTokenData, TokenKey));
+			check(TokenKeyHandle.IsValid());
+		}
+	}
+
+	// If owning blueprint is null, then we may not be customizing the naming tokens directly, such as a settings object.
+	// or the tokens are predefined and can't be extended in BP.
+	
+	HeaderRow
+	.NameContent()
+	[
+		PropertyHandle->CreatePropertyNameWidget()
+	]
+	.ValueContent()
+	[
+		PropertyHandle->CreatePropertyValueWidget(/* bDisplayDefaultPropertyButtons (adds identical options) */ false)
+	];
+}
+
+void FNamingTokensCustomization::CustomizeChildren(TSharedRef<IPropertyHandle> PropertyHandle, IDetailChildrenBuilder& ChildBuilder,
+	IPropertyTypeCustomizationUtils& CustomizationUtils)
+{
+	const bool bShouldApplyCustomization = OwningBlueprint.IsValid();
+
+	if (bShouldApplyCustomization)
+	{
+		const TArray<UFunction*> Functions = GetAvailableFunctions();
+		for (const UFunction* Function : Functions)
+		{
+			FunctionNames.Add(MakeShared<FString>(Function->GetName()));
+		}
+
+		FString CurrentValue;
+		FunctionNameHandle->GetValue(CurrentValue);
+
+		// Default selected function
+		for (const TSharedPtr<FString>& Name : FunctionNames)
+		{
+			if (*Name == CurrentValue)
+			{
+				SelectedFunctionName = Name;
+				break;
+			}
+		}
+
+		FunctionNameHandle->MarkHiddenByCustomization();
+	}
+	
+	uint32 NumChildren = 0;
+	PropertyHandle->GetNumChildren(NumChildren);
+	
+	for (uint32 ChildNum = 0; ChildNum < NumChildren; ++ChildNum)
+	{
+		const TSharedPtr<IPropertyHandle> ChildHandle = PropertyHandle->GetChildHandle(ChildNum);
+		if (!ChildHandle.IsValid() || (bShouldApplyCustomization && ChildHandle->GetProperty() == FunctionNameHandle->GetProperty()))
+		{
+			continue;
+		}
+		ChildBuilder.AddProperty(ChildHandle.ToSharedRef());
+	}
+
+	if (!bShouldApplyCustomization)
+	{
+		return;
+	}
+	
+	FDetailWidgetRow& HeaderRow = ChildBuilder.AddProperty(FunctionNameHandle.ToSharedRef()).CustomWidget();
+	HeaderRow.NameContent()
+	[
+		FunctionNameHandle->CreatePropertyNameWidget()
+	];
+	HeaderRow.ValueContent()
+	[
+		SNew(SHorizontalBox)
+		+SHorizontalBox::Slot()
+		[
+			SNew(SComboBox<TSharedPtr<FString>>)
+			.OptionsSource(&FunctionNames)
+			.OnGenerateWidget(this, &FNamingTokensCustomization::MakeComboBoxWidget)
+			.OnSelectionChanged(this, &FNamingTokensCustomization::OnFunctionSelected, FunctionNameHandle)
+			.InitiallySelectedItem(SelectedFunctionName)
+			[
+				SNew(STextBlock)
+				.Text(this, &FNamingTokensCustomization::GetSelectedFunctionText)
+			]
+		]
+		+ SHorizontalBox::Slot()
+		.AutoWidth()
+		.VAlign(VAlign_Center)
+		[
+			SNew(SButton)
+			[
+				SNew(SImage)
+				.Image(FAppStyle::Get().GetBrush("PListEditor.Button_AddToArray"))
+			]
+			.ButtonStyle(FAppStyle::Get(), "HoverHintOnly")
+			.ToolTipText(LOCTEXT("AddFunction", "Create and assign a new function graph for evaluating tokens."))
+			.OnClicked(this, &FNamingTokensCustomization::OnAddFunctionClicked)
+			.IsEnabled(this, &FNamingTokensCustomization::CanAddFunction)
+		]
+	];
+}
+
+UNamingTokens* FNamingTokensCustomization::GetOwningNamingTokens() const
+{
+	if (CustomizationUtilsPtr)
+	{
+		const TSharedPtr<IPropertyUtilities> PropertyUtilities = CustomizationUtilsPtr->GetPropertyUtilities();
+		if (PropertyUtilities.IsValid())
+		{
+			const TArray<TWeakObjectPtr<UObject>> ObjectsBeingCustomized = PropertyUtilities->GetSelectedObjects();
+
+			if (ObjectsBeingCustomized.Num() == 1)
+			{
+				if (UNamingTokens* Tokens = Cast<UNamingTokens>(ObjectsBeingCustomized[0]))
+				{
+					return Tokens;
+				}
+			}
+		}
+	}
+	return nullptr;
+}
+
+TArray<UFunction*> FNamingTokensCustomization::GetAvailableFunctions() const
+{
+	TArray<UFunction*> Result;
+	if (OwningBlueprint.IsValid())
+	{
+		for (TFieldIterator<UFunction> It(OwningBlueprint->SkeletonGeneratedClass ?
+			OwningBlueprint->SkeletonGeneratedClass : OwningBlueprint->GeneratedClass,
+			EFieldIteratorFlags::IncludeSuper); It; ++It)
+		{
+			if (UE::NamingTokens::Utils::ValidateTokenFunction(*It))
+			{
+				Result.Add(*It);
+			}
+		}
+	}
+
+	return Result;
+}
+
+TSharedRef<SWidget> FNamingTokensCustomization::MakeComboBoxWidget(TSharedPtr<FString> InItem)
+{
+	return SNew(STextBlock)
+		.Text(FText::FromString(*InItem));
+}
+
+void FNamingTokensCustomization::OnFunctionSelected(TSharedPtr<FString> NewValue, ESelectInfo::Type SelectInfo,
+	TSharedPtr<IPropertyHandle> InFunctionNameHandle)
+{
+	if (NewValue.IsValid())
+	{
+		SelectedFunctionName = NewValue;
+		InFunctionNameHandle->SetValue(*NewValue);
+	}
+}
+
+FText FNamingTokensCustomization::GetSelectedFunctionText() const
+{
+	return SelectedFunctionName.IsValid() ? FText::FromString(*SelectedFunctionName) : FText::FromString(TEXT("Select Function"));
+}
+
+FReply FNamingTokensCustomization::OnAddFunctionClicked()
+{
+	if (UBlueprint* Blueprint = OwningBlueprint.Get())
+	{
+		FString TokenKey;
+		TokenKeyHandle->GetValue(TokenKey);
+
+		const FName NewFunctionName = UE::NamingTokens::Utils::Editor::Private::CreateNewTokenGraph(Blueprint, TokenKey);
+		FunctionNameHandle->SetValue(NewFunctionName);
+	}
+	
+	return FReply::Handled();
+}
+
+bool FNamingTokensCustomization::CanAddFunction() const
+{
+	FString TokenKey;
+	if (TokenKeyHandle.IsValid())
+	{
+		TokenKeyHandle->GetValue(TokenKey);
+		
+		const FString BaseFunctionName = UE::NamingTokens::Utils::Editor::Private::CreateBaseTokenFunctionName(TokenKey);
+
+		// Make sure the name isn't already being used.
+		for (const TSharedPtr<FString>& FunctionName : FunctionNames)
+		{
+			if (FunctionName.IsValid() && **FunctionName == BaseFunctionName)
+			{
+				return false;
+			}
+		}
+	}
+	
+	return !TokenKey.IsEmpty();
+}
+
+#undef LOCTEXT_NAMESPACE
diff --git a/Engine/Plugins/Developer/NamingTokens/Source/NamingTokensUncookedOnly/Private/Customization/NamingTokensCustomization.h b/Engine/Plugins/Developer/NamingTokens/Source/NamingTokensUncookedOnly/Private/Customization/NamingTokensCustomization.h
new file mode 100644
--- /dev/null
+++ b/Engine/Plugins/Developer/NamingTokens/Source/NamingTokensUncookedOnly/Private/Customization/NamingTokensCustomization.h
@@ -0,0 +1,57 @@
+// Copyright Epic Games, Inc. All Rights Reserved.
+
+#pragma once
+
+#include "IPropertyTypeCustomization.h"
+#include "CoreMinimal.h"
+#include "Input/Reply.h"
+
+class UNamingTokens;
+class SWidget;
+
+class FNamingTokensCustomization final : public IPropertyTypeCustomization
+{
+public:
+	static TSharedRef<IPropertyTypeCustomization> MakeInstance()
+	{
+		return MakeShareable(new FNamingTokensCustomization());
+	}
+
+	virtual void CustomizeHeader(TSharedRef<IPropertyHandle> PropertyHandle, FDetailWidgetRow& HeaderRow, IPropertyTypeCustomizationUtils& CustomizationUtils) override;
+	virtual void CustomizeChildren(TSharedRef<IPropertyHandle> PropertyHandle, IDetailChildrenBuilder& ChildBuilder, IPropertyTypeCustomizationUtils& CustomizationUtils) override;
+
+	/** The naming tokens owning this customization. */
+	UNamingTokens* GetOwningNamingTokens() const;
+
+private:
+	/** Available functions which can be assigned to naming tokens. */
+	TArray<UFunction*> GetAvailableFunctions() const;
+	/** Make the combo box for selecting functions. */
+	TSharedRef<SWidget> MakeComboBoxWidget(TSharedPtr<FString> InItem);
+	/** When the user selects a function. */
+	void OnFunctionSelected(
+		TSharedPtr<FString> NewValue, 
+		ESelectInfo::Type SelectInfo, 
+		TSharedPtr<IPropertyHandle> InFunctionNameHandle);
+	/** The text of the currently selected function. */
+	FText GetSelectedFunctionText() const;
+	/** User clicked the button to add a new function. */
+	FReply OnAddFunctionClicked();
+	/** If a new function can be added. */
+	bool CanAddFunction() const;
+
+private:
+	/** Customization pointer. */
+	IPropertyTypeCustomizationUtils* CustomizationUtilsPtr = nullptr;
+	/** Blueprint of the naming tokens. */
+	TWeakObjectPtr<UBlueprint> OwningBlueprint;
+
+	/** All available function names. */
+	TArray<TSharedPtr<FString>> FunctionNames;
+	/** The selected function name. */
+	TSharedPtr<FString> SelectedFunctionName;
+	/** Property handle to the function name. */
+	TSharedPtr<IPropertyHandle> FunctionNameHandle;
+	/** Property handle to the token key. */
+	TSharedPtr<IPropertyHandle> TokenKeyHandle;
+};
\ No newline at end of file
diff --git a/Engine/Plugins/Developer/NamingTokens/Source/NamingTokensUncookedOnly/Private/NamingTokensFactory.cpp b/Engine/Plugins/Developer/NamingTokens/Source/NamingTokensUncookedOnly/Private/NamingTokensFactory.cpp
new file mode 100644
--- /dev/null
+++ b/Engine/Plugins/Developer/NamingTokens/Source/NamingTokensUncookedOnly/Private/NamingTokensFactory.cpp
@@ -0,0 +1,67 @@
+// Copyright Epic Games, Inc. All Rights Reserved.
+
+#include "NamingTokensFactory.h"
+
+#include "NamingTokens.h"
+
+#include "BlueprintEditorSettings.h"
+#include "EdGraphSchema_K2.h"
+#include "Engine/Blueprint.h"
+#include "Engine/BlueprintGeneratedClass.h"
+#include "K2Node_Event.h"
+#include "Kismet2/KismetEditorUtilities.h"
+
+UNamingTokensFactory::UNamingTokensFactory()
+{
+	SupportedClass = UNamingTokens::StaticClass();
+
+	bCreateNew = true;
+	bEditAfterNew = true;
+}
+
+UObject* UNamingTokensFactory::FactoryCreateNew(UClass* InClass, UObject* InParent, FName InName, EObjectFlags Flags, UObject* Context,
+	FFeedbackContext* Warn)
+{
+	UClass* ParentClass = SupportedClass;
+	check(InClass->IsChildOf(ParentClass));
+    
+	UBlueprint* NewBP = FKismetEditorUtilities::CreateBlueprint(
+		ParentClass,
+		InParent,
+		InName,
+		BPTYPE_Normal,
+		UBlueprint::StaticClass(),
+		UBlueprintGeneratedClass::StaticClass()
+	);
+	
+	const UBlueprintEditorSettings* Settings = GetDefault<UBlueprintEditorSettings>();
+	if (Settings && Settings->bSpawnDefaultBlueprintNodes)
+	{
+		// Create default events.
+		int32 NodePositionY = 0;
+
+		UEdGraph* EventGraph = FindObject<UEdGraph>(NewBP, *(UEdGraphSchema_K2::GN_EventGraph.ToString()));
+		check(EventGraph);
+		
+		if (UK2Node_Event* OnPreEvaluateNode = FKismetEditorUtilities::AddDefaultEventNode(NewBP, EventGraph,
+			GET_FUNCTION_NAME_CHECKED(UNamingTokens, OnPreEvaluate), UNamingTokens::StaticClass(), NodePositionY))
+		{
+			// Set the node comment. The comment is displayed because we are automatically placed as ghost nodes.
+			// Once a connection is made the comment will go away, but is still accessible via tooltip.
+			const UFunction* Function = UNamingTokens::StaticClass()->FindFunctionByName(GET_FUNCTION_NAME_CHECKED(UNamingTokens, OnPreEvaluate));
+			check(Function);
+			OnPreEvaluateNode->NodeComment = Function->GetToolTipText().ToString();
+		}
+		if (UK2Node_Event* OnPostEvaluateNode = FKismetEditorUtilities::AddDefaultEventNode(NewBP, EventGraph,
+			GET_FUNCTION_NAME_CHECKED(UNamingTokens, OnPostEvaluate), UNamingTokens::StaticClass(), NodePositionY))
+		{
+			// Set the node comment. The comment is displayed because we are automatically placed as ghost nodes.
+			// Once a connection is made the comment will go away, but is still accessible via tooltip.
+			const UFunction* Function = UNamingTokens::StaticClass()->FindFunctionByName(GET_FUNCTION_NAME_CHECKED(UNamingTokens, OnPostEvaluate));
+			check(Function);
+			OnPostEvaluateNode->NodeComment = Function->GetToolTipText().ToString();
+		}
+	}
+
+	return NewBP;
+}
diff --git a/Engine/Plugins/Developer/NamingTokens/Source/NamingTokensUncookedOnly/Private/NamingTokensFactory.h b/Engine/Plugins/Developer/NamingTokens/Source/NamingTokensUncookedOnly/Private/NamingTokensFactory.h
new file mode 100644
--- /dev/null
+++ b/Engine/Plugins/Developer/NamingTokens/Source/NamingTokensUncookedOnly/Private/NamingTokensFactory.h
@@ -0,0 +1,19 @@
+// Copyright Epic Games, Inc. All Rights Reserved.
+
+#pragma once
+
+#include "CoreMinimal.h"
+#include "Factories/Factory.h"
+
+#include "NamingTokensFactory.generated.h"
+
+UCLASS(MinimalAPI)
+class UNamingTokensFactory : public UFactory
+{
+	GENERATED_BODY()
+
+public:
+	UNamingTokensFactory();
+	
+	virtual UObject* FactoryCreateNew(UClass* InClass, UObject* InParent, FName InName, EObjectFlags Flags, UObject* Context, FFeedbackContext* Warn) override;
+};
\ No newline at end of file
diff --git a/Engine/Plugins/Developer/NamingTokens/Source/NamingTokensUncookedOnly/Private/NamingTokensUncookedOnlyModule.cpp b/Engine/Plugins/Developer/NamingTokens/Source/NamingTokensUncookedOnly/Private/NamingTokensUncookedOnlyModule.cpp
new file mode 100644
--- /dev/null
+++ b/Engine/Plugins/Developer/NamingTokens/Source/NamingTokensUncookedOnly/Private/NamingTokensUncookedOnlyModule.cpp
@@ -0,0 +1,28 @@
+﻿// Copyright Epic Games, Inc. All Rights Reserved.
+
+#include "NamingTokensUncookedOnlyModule.h"
+
+#include "Customization/NamingTokensCustomization.h"
+#include "NamingTokenData.h"
+
+#include "PropertyEditorModule.h"
+
+void FNamingTokensUncookedOnlyModule::StartupModule()
+{
+	FPropertyEditorModule& PropertyModule = FModuleManager::GetModuleChecked<FPropertyEditorModule>("PropertyEditor");
+	PropertyModule.RegisterCustomPropertyTypeLayout(FNamingTokenData::StaticStruct()->GetFName(),
+		FOnGetPropertyTypeCustomizationInstance::CreateStatic(&FNamingTokensCustomization::MakeInstance));
+}
+
+void FNamingTokensUncookedOnlyModule::ShutdownModule()
+{
+	if (UObjectInitialized() && !IsEngineExitRequested())
+	{
+		if (FPropertyEditorModule* PropertyEditorModule = FModuleManager::GetModulePtr<FPropertyEditorModule>("PropertyEditor"))
+		{
+			PropertyEditorModule->UnregisterCustomPropertyTypeLayout(FNamingTokenData::StaticStruct()->GetFName());
+		}
+	}
+}
+
+IMPLEMENT_MODULE(FNamingTokensUncookedOnlyModule, NamingTokensUncookedOnly)
\ No newline at end of file
diff --git a/Engine/Plugins/Developer/NamingTokens/Source/NamingTokensUncookedOnly/Private/Utils/NamingTokensEditorUtils.cpp b/Engine/Plugins/Developer/NamingTokens/Source/NamingTokensUncookedOnly/Private/Utils/NamingTokensEditorUtils.cpp
new file mode 100644
--- /dev/null
+++ b/Engine/Plugins/Developer/NamingTokens/Source/NamingTokensUncookedOnly/Private/Utils/NamingTokensEditorUtils.cpp
@@ -0,0 +1,31 @@
+// Copyright Epic Games, Inc. All Rights Reserved.
+
+#include "NamingTokensEditorUtils.h"
+
+#include "Kismet2/BlueprintEditorUtils.h"
+#include "Kismet2/KismetEditorUtilities.h"
+#include "Utils/NamingTokenUtils.h"
+
+FString UE::NamingTokens::Utils::Editor::Private::CreateBaseTokenFunctionName(const FString& InTokenKey)
+{
+	return TEXT("ProcessToken_") + InTokenKey;
+}
+
+FName UE::NamingTokens::Utils::Editor::Private::CreateNewTokenGraph(UBlueprint* InBlueprint, const FString& InTokenKey)
+{
+	check(InBlueprint);
+	
+	const FString BaseName = CreateBaseTokenFunctionName(InTokenKey);
+	const FName FunctionName = FBlueprintEditorUtils::GenerateUniqueGraphName(InBlueprint, BaseName);
+
+	InBlueprint->Modify();
+		
+	UEdGraph* NewGraph = FBlueprintEditorUtils::CreateNewGraph(InBlueprint, FunctionName, UEdGraph::StaticClass(), UEdGraphSchema_K2::StaticClass());
+	
+	UFunction* FunctionSignature = GetProcessTokenFunctionSignature();
+		
+	FBlueprintEditorUtils::AddFunctionGraph<UFunction>(InBlueprint, NewGraph, /*bIsUserCreated=*/ true, FunctionSignature);
+	FKismetEditorUtilities::BringKismetToFocusAttentionOnObject(NewGraph);
+
+	return FunctionName;
+}
diff --git a/Engine/Plugins/Developer/NamingTokens/Source/NamingTokensUncookedOnly/Private/Utils/NamingTokensEditorUtils.h b/Engine/Plugins/Developer/NamingTokens/Source/NamingTokensUncookedOnly/Private/Utils/NamingTokensEditorUtils.h
new file mode 100644
--- /dev/null
+++ b/Engine/Plugins/Developer/NamingTokens/Source/NamingTokensUncookedOnly/Private/Utils/NamingTokensEditorUtils.h
@@ -0,0 +1,16 @@
+// Copyright Epic Games, Inc. All Rights Reserved.
+
+#pragma once
+
+#include "CoreMinimal.h"
+
+class UBlueprint;
+
+namespace UE::NamingTokens::Utils::Editor::Private
+{
+	/** Create the default function token name to use, before unique naming adjustments. */
+	FString CreateBaseTokenFunctionName(const FString& InTokenKey);
+	
+	/** Creates a new k2 graph in a blueprint for a given token key. */
+	FName CreateNewTokenGraph(UBlueprint* InBlueprint, const FString& InTokenKey);
+}
diff --git a/Engine/Plugins/Developer/NamingTokens/Source/NamingTokensUncookedOnly/Public/NamingTokensUncookedOnlyModule.h b/Engine/Plugins/Developer/NamingTokens/Source/NamingTokensUncookedOnly/Public/NamingTokensUncookedOnlyModule.h
new file mode 100644
--- /dev/null
+++ b/Engine/Plugins/Developer/NamingTokens/Source/NamingTokensUncookedOnly/Public/NamingTokensUncookedOnlyModule.h
@@ -0,0 +1,13 @@
+﻿// Copyright Epic Games, Inc. All Rights Reserved.
+
+#pragma once
+
+#include "CoreMinimal.h"
+#include "Modules/ModuleManager.h"
+
+class FNamingTokensUncookedOnlyModule : public IModuleInterface
+{
+public:
+    virtual void StartupModule() override;
+    virtual void ShutdownModule() override;
+};
diff --git a/Engine/Plugins/VirtualProduction/CaptureManager/CaptureManagerApp/Source/CaptureDataConverter/Private/Nodes/ThirdPartyEncoder/CaptureConvertAudioDataThirdParty.cpp b/Engine/Plugins/VirtualProduction/CaptureManager/CaptureManagerApp/Source/CaptureDataConverter/Private/Nodes/ThirdPartyEncoder/CaptureConvertAudioDataThirdParty.cpp
--- a/Engine/Plugins/VirtualProduction/CaptureManager/CaptureManagerApp/Source/CaptureDataConverter/Private/Nodes/ThirdPartyEncoder/CaptureConvertAudioDataThirdParty.cpp
+++ b/Engine/Plugins/VirtualProduction/CaptureManager/CaptureManagerApp/Source/CaptureDataConverter/Private/Nodes/ThirdPartyEncoder/CaptureConvertAudioDataThirdParty.cpp
@@ -13,10 +13,10 @@
 
 #include "CaptureManagerMediaRWModule.h"
 
-#include "Editor.h"
+#include "Engine/Engine.h"
 #include "GlobalNamingTokens.h"
 #include "NamingTokenData.h"
-#include "NamingTokensEditorSubsystem.h"
+#include "NamingTokensEngineSubsystem.h"
 #include "Settings/CaptureManagerSettings.h"
 #include "Settings/CaptureManagerTemplateTokens.h"
 
@@ -110,7 +110,7 @@ FCaptureConvertAudioDataThirdParty::FResult FCaptureConvertAudioDataThirdParty::
 		{
 			using namespace UE::CaptureManager;
 
-			UNamingTokensEditorSubsystem* NamingTokensSubsystem = GEditor->GetEditorSubsystem<UNamingTokensEditorSubsystem>();
+			UNamingTokensEngineSubsystem* NamingTokensSubsystem = GEngine->GetEngineSubsystem<UNamingTokensEngineSubsystem>();
 			check(NamingTokensSubsystem);
 
 			const UCaptureManagerSettings* Settings = GetDefault<UCaptureManagerSettings>();
diff --git a/Engine/Plugins/VirtualProduction/CaptureManager/CaptureManagerApp/Source/CaptureDataConverter/Private/Nodes/ThirdPartyEncoder/CaptureConvertVideoDataThirdParty.cpp b/Engine/Plugins/VirtualProduction/CaptureManager/CaptureManagerApp/Source/CaptureDataConverter/Private/Nodes/ThirdPartyEncoder/CaptureConvertVideoDataThirdParty.cpp
--- a/Engine/Plugins/VirtualProduction/CaptureManager/CaptureManagerApp/Source/CaptureDataConverter/Private/Nodes/ThirdPartyEncoder/CaptureConvertVideoDataThirdParty.cpp
+++ b/Engine/Plugins/VirtualProduction/CaptureManager/CaptureManagerApp/Source/CaptureDataConverter/Private/Nodes/ThirdPartyEncoder/CaptureConvertVideoDataThirdParty.cpp
@@ -12,10 +12,10 @@
 #include "Nodes/CaptureCopyProgressReporter.h"
 #include "CaptureThirdPartyNodeUtils.h"
 
-#include "Editor.h"
+#include "Engine/Engine.h"
 #include "GlobalNamingTokens.h"
 #include "NamingTokenData.h"
-#include "NamingTokensEditorSubsystem.h"
+#include "NamingTokensEngineSubsystem.h"
 #include "Settings/CaptureManagerSettings.h"
 #include "Settings/CaptureManagerTemplateTokens.h"
 
@@ -179,7 +179,7 @@ FCaptureConvertVideoDataThirdParty::FResult FCaptureConvertVideoDataThirdParty::
 		{
 			using namespace UE::CaptureManager;
 
-			UNamingTokensEditorSubsystem* NamingTokensSubsystem = GEditor->GetEditorSubsystem<UNamingTokensEditorSubsystem>();
+			UNamingTokensEngineSubsystem* NamingTokensSubsystem = GEngine->GetEngineSubsystem<UNamingTokensEngineSubsystem>();
 			check(NamingTokensSubsystem);
 
 			const UCaptureManagerSettings* Settings = GetDefault<UCaptureManagerSettings>();
diff --git a/Engine/Plugins/VirtualProduction/CaptureManager/CaptureManagerApp/Source/CaptureManagerEditor/Private/IngestManagement/IngestJobExecutor.cpp b/Engine/Plugins/VirtualProduction/CaptureManager/CaptureManagerApp/Source/CaptureManagerEditor/Private/IngestManagement/IngestJobExecutor.cpp
--- a/Engine/Plugins/VirtualProduction/CaptureManager/CaptureManagerApp/Source/CaptureManagerEditor/Private/IngestManagement/IngestJobExecutor.cpp
+++ b/Engine/Plugins/VirtualProduction/CaptureManager/CaptureManagerApp/Source/CaptureManagerEditor/Private/IngestManagement/IngestJobExecutor.cpp
@@ -16,7 +16,7 @@
 #include "Editor.h"
 #include "GlobalNamingTokens.h"
 #include "NamingTokenData.h"
-#include "NamingTokensEditorSubsystem.h"
+#include "NamingTokensEngineSubsystem.h"
 #include "Settings/CaptureManagerTemplateTokens.h"
 
 DEFINE_LOG_CATEGORY_STATIC(LogCaptureManagerIngestJobExecutor, Log, All);
@@ -193,7 +193,7 @@ UIngestCapability_ProcessHandle* FIngestJobExecutor::RunIngest(FIngestJob& InJob
 	CallOnGameThread(
 		[&InJob, &DeviceName, IngestOptions]()
 		{
-			UNamingTokensEditorSubsystem* NamingTokensSubsystem = GEditor->GetEditorSubsystem<UNamingTokensEditorSubsystem>();
+			UNamingTokensEngineSubsystem* NamingTokensSubsystem = GEngine->GetEngineSubsystem<UNamingTokensEngineSubsystem>();
 			check(NamingTokensSubsystem);
 
 			FNamingTokenFilterArgs GeneralNamingTokenArgs;
diff --git a/Engine/Plugins/VirtualProduction/CaptureManager/CaptureManagerApp/Source/CaptureManagerSettings/Private/Settings/CaptureManagerSettingsCustomization.cpp b/Engine/Plugins/VirtualProduction/CaptureManager/CaptureManagerApp/Source/CaptureManagerSettings/Private/Settings/CaptureManagerSettingsCustomization.cpp
--- a/Engine/Plugins/VirtualProduction/CaptureManager/CaptureManagerApp/Source/CaptureManagerSettings/Private/Settings/CaptureManagerSettingsCustomization.cpp
+++ b/Engine/Plugins/VirtualProduction/CaptureManager/CaptureManagerApp/Source/CaptureManagerSettings/Private/Settings/CaptureManagerSettingsCustomization.cpp
@@ -6,8 +6,8 @@
 #include "Settings/CaptureManagerSettings.h"
 #include "DetailLayoutBuilder.h"
 #include "DetailWidgetRow.h"
-#include "Editor.h"
-#include "NamingTokensEditorSubsystem.h"
+#include "Engine/Engine.h"
+#include "NamingTokensEngineSubsystem.h"
 #include "Widgets/Input/SEditableText.h"
 #include "Widgets/Text/STextBlock.h"
 
@@ -94,9 +94,9 @@ void FCaptureManagerSettingsCustomization::CustomizeDetails(IDetailLayoutBuilder
 FText FCaptureManagerSettingsCustomization::GetDisplayTokenText(FNamingTokenFilterArgs InArgs) const
 {
 	FString FormattedTokensString = TEXT("None");
-	if (GEditor)
+	if (GEngine)
 	{
-		FormattedTokensString = GEditor->GetEditorSubsystem<UNamingTokensEditorSubsystem>()->GetFormattedTokensStringForDisplay(MoveTemp(InArgs));
+		FormattedTokensString = GEngine->GetEngineSubsystem<UNamingTokensEngineSubsystem>()->GetFormattedTokensStringForDisplay(MoveTemp(InArgs));
 	}
 	return FText::FromString(FormattedTokensString);
 }
diff --git a/Engine/Plugins/VirtualProduction/CaptureManager/CaptureManagerEditor/Source/CaptureManagerEditorSettings/Private/Settings/CaptureManagerEditorSettingsCustomization.cpp b/Engine/Plugins/VirtualProduction/CaptureManager/CaptureManagerEditor/Source/CaptureManagerEditorSettings/Private/Settings/CaptureManagerEditorSettingsCustomization.cpp
--- a/Engine/Plugins/VirtualProduction/CaptureManager/CaptureManagerEditor/Source/CaptureManagerEditorSettings/Private/Settings/CaptureManagerEditorSettingsCustomization.cpp
+++ b/Engine/Plugins/VirtualProduction/CaptureManager/CaptureManagerEditor/Source/CaptureManagerEditorSettings/Private/Settings/CaptureManagerEditorSettingsCustomization.cpp
@@ -7,7 +7,8 @@
 #include "DetailLayoutBuilder.h"
 #include "DetailWidgetRow.h"
 #include "Editor.h"
-#include "NamingTokensEditorSubsystem.h"
+#include "Engine/Engine.h"
+#include "NamingTokensEngineSubsystem.h"
 #include "Widgets/Input/SEditableText.h"
 #include "Widgets/Text/STextBlock.h"
 
@@ -121,9 +122,9 @@ void FCaptureManagerEditorSettingsCustomization::CustomizeDetails(IDetailLayoutB
 FText FCaptureManagerEditorSettingsCustomization::GetDisplayTokenText(FNamingTokenFilterArgs InArgs) const
 {
 	FString FormattedTokensString = TEXT("None");
-	if (GEditor)
+	if (GEngine)
 	{
-		FormattedTokensString = GEditor->GetEditorSubsystem<UNamingTokensEditorSubsystem>()->GetFormattedTokensStringForDisplay(MoveTemp(InArgs));
+		FormattedTokensString = GEngine->GetEngineSubsystem<UNamingTokensEngineSubsystem>()->GetFormattedTokensStringForDisplay(MoveTemp(InArgs));
 	}
 	return FText::FromString(FormattedTokensString);
 }
diff --git a/Engine/Plugins/VirtualProduction/CaptureManager/CaptureManagerEditor/Source/LiveLinkHubWorkerManager/Private/IngestProcess/IngestCaptureDataProcess.cpp b/Engine/Plugins/VirtualProduction/CaptureManager/CaptureManagerEditor/Source/LiveLinkHubWorkerManager/Private/IngestProcess/IngestCaptureDataProcess.cpp
--- a/Engine/Plugins/VirtualProduction/CaptureManager/CaptureManagerEditor/Source/LiveLinkHubWorkerManager/Private/IngestProcess/IngestCaptureDataProcess.cpp
+++ b/Engine/Plugins/VirtualProduction/CaptureManager/CaptureManagerEditor/Source/LiveLinkHubWorkerManager/Private/IngestProcess/IngestCaptureDataProcess.cpp
@@ -11,9 +11,10 @@
 #include "Async/HelperFunctions.h"
 
 #include "Editor.h"
+#include "Engine/Engine.h"
 #include "GlobalNamingTokens.h"
 #include "NamingTokenData.h"
-#include "NamingTokensEditorSubsystem.h"
+#include "NamingTokensEngineSubsystem.h"
 
 #include "Settings/CaptureManagerEditorSettings.h"
 #include "Settings/CaptureManagerEditorTemplateTokens.h"
@@ -140,7 +141,7 @@ UE::CaptureManager::FCreateAssetsData FIngestCaptureDataProcess::PrepareAssetsDa
 	CallOnGameThread(
 		[&CreateAssetData, &InIngestCaptureData, &Settings, &ImportNamedArgs]()
 		{
-			UNamingTokensEditorSubsystem* NamingTokensSubsystem = GEditor->GetEditorSubsystem<UNamingTokensEditorSubsystem>();
+			UNamingTokensEngineSubsystem* NamingTokensSubsystem = GEngine->GetEngineSubsystem<UNamingTokensEngineSubsystem>();
 			check(NamingTokensSubsystem);
 
 			FNamingTokenFilterArgs NamingTokenArgs;
diff --git a/Engine/Plugins/VirtualProduction/CaptureManager/CaptureManagerEditor/Source/LiveLinkHubWorkerManager/Private/LiveLinkHubWorker.cpp b/Engine/Plugins/VirtualProduction/CaptureManager/CaptureManagerEditor/Source/LiveLinkHubWorkerManager/Private/LiveLinkHubWorker.cpp
--- a/Engine/Plugins/VirtualProduction/CaptureManager/CaptureManagerEditor/Source/LiveLinkHubWorkerManager/Private/LiveLinkHubWorker.cpp
+++ b/Engine/Plugins/VirtualProduction/CaptureManager/CaptureManagerEditor/Source/LiveLinkHubWorkerManager/Private/LiveLinkHubWorker.cpp
@@ -27,9 +27,10 @@
 #include "SoundWaveTimecodeUtils.h"
 
 #include "Editor.h"
+#include "Engine/Engine.h"
 #include "GlobalNamingTokens.h"
 #include "NamingTokenData.h"
-#include "NamingTokensEditorSubsystem.h"
+#include "NamingTokensEngineSubsystem.h"
 
 #include "Settings/CaptureManagerEditorSettings.h"
 #include "Settings/CaptureManagerEditorTemplateTokens.h"
@@ -105,7 +106,7 @@ void FLiveLinkHubImportWorker::EvaluateSettings(const FUploadDataHeader& InHeade
 	CallOnGameThread(
 		[&DataStorageOut, &GeneralNamedArgs, Settings]()
 		{
-			UNamingTokensEditorSubsystem* NamingTokensSubsystem = GEditor->GetEditorSubsystem<UNamingTokensEditorSubsystem>();
+			UNamingTokensEngineSubsystem* NamingTokensSubsystem = GEngine->GetEngineSubsystem<UNamingTokensEngineSubsystem>();
 			check(NamingTokensSubsystem);
 
 			FNamingTokenFilterArgs NamingTokenArgs;
diff --git a/Engine/Plugins/VirtualProduction/Shots/Source/ShotsEditor/Private/ProductionSettings.cpp b/Engine/Plugins/VirtualProduction/Shots/Source/ShotsEditor/Private/ProductionSettings.cpp
--- a/Engine/Plugins/VirtualProduction/Shots/Source/ShotsEditor/Private/ProductionSettings.cpp
+++ b/Engine/Plugins/VirtualProduction/Shots/Source/ShotsEditor/Private/ProductionSettings.cpp
@@ -2,14 +2,14 @@
 #include "ProductionSettings.h"
 
 #include "AssetToolsModule.h"
-#include "Editor.h"
+#include "Engine/Engine.h"
 #include "HAL/FileManager.h"
 #include "HAL/PlatformFileManager.h"
 #include "LevelSequenceProjectSettings.h"
 #include "Misc/ConfigCacheIni.h"
 #include "Misc/CoreDelegates.h"
 #include "MovieSceneToolsProjectSettings.h"
-#include "NamingTokensEditorSubsystem.h"
+#include "NamingTokensEngineSubsystem.h"
 #include "Sections/MovieSceneSubSection.h"
 #include "SourceControlHelpers.h"
 
@@ -91,7 +91,7 @@ void UProductionSettings::PostInitProperties()
 	// Register a filter with the naming tokens subsystem to apply the active production's namespace deny list when evaluating tokens
 	auto RegisterNamespaceFilter = [this]()
 		{
-			UNamingTokensEditorSubsystem* NamingTokensSubsystem = GEditor->GetEditorSubsystem<UNamingTokensEditorSubsystem>();
+			UNamingTokensEngineSubsystem* NamingTokensSubsystem = GEngine->GetEngineSubsystem<UNamingTokensEngineSubsystem>();
 			NamingTokensSubsystem->RegisterNamespaceFilter("CinematicProductionSettings", FFilterNamespace::CreateUObject(this, &UProductionSettings::FilterNamingTokenNamespaces));
 		};
 
diff --git a/Engine/Plugins/VirtualProduction/Shots/Source/ShotsEditor/Private/UI/SNamingTokensPanel.cpp b/Engine/Plugins/VirtualProduction/Shots/Source/ShotsEditor/Private/UI/SNamingTokensPanel.cpp
--- a/Engine/Plugins/VirtualProduction/Shots/Source/ShotsEditor/Private/UI/SNamingTokensPanel.cpp
+++ b/Engine/Plugins/VirtualProduction/Shots/Source/ShotsEditor/Private/UI/SNamingTokensPanel.cpp
@@ -3,10 +3,10 @@
 #include "SNamingTokensPanel.h"
 
 #include "Algo/Reverse.h"
-#include "Editor.h"
+#include "Engine/Engine.h"
 #include "GlobalNamingTokens.h"
 #include "NamingTokenData.h"
-#include "NamingTokensEditorSubsystem.h"
+#include "NamingTokensEngineSubsystem.h"
 #include "ProductionSettings.h"
 #include "ShotsEditorStyle.h"
 #include "UI/SActiveProductionCombo.h"
@@ -19,7 +19,7 @@
 
 void SNamingTokensPanel::Construct(const FArguments& InArgs)
 {
-	const UNamingTokensEditorSubsystem* const NamingTokensSubsystem = GEditor->GetEditorSubsystem<UNamingTokensEditorSubsystem>();
+	const UNamingTokensEngineSubsystem* const NamingTokensSubsystem = GEngine->GetEngineSubsystem<UNamingTokensEngineSubsystem>();
 
 	// Build the list of the global tokens
 	const UNamingTokens* const GlobalTokens = NamingTokensSubsystem->GetNamingTokens(UGlobalNamingTokens::GetGlobalNamespace());
@@ -264,7 +264,7 @@ void SNamingTokensPanel::OnNamespaceSelectionChanged(TSharedPtr<FString> Selecte
 
 	if (SelectedNamespace)
 	{
-		UNamingTokensEditorSubsystem* NamingTokensSubsystem = GEditor->GetEditorSubsystem<UNamingTokensEditorSubsystem>();
+		UNamingTokensEngineSubsystem* NamingTokensSubsystem = GEngine->GetEngineSubsystem<UNamingTokensEngineSubsystem>();
 		if (UNamingTokens* CustomTokens = NamingTokensSubsystem->GetNamingTokens(*SelectedNamespace))
 		{
 			const TArray<FNamingTokenData>& CustomTokenData = CustomTokens->GetAllTokens();
diff --git a/Engine/Plugins/VirtualProduction/Takes/Source/TakeRecorder/Private/Widgets/STakeRecorderCockpit.cpp b/Engine/Plugins/VirtualProduction/Takes/Source/TakeRecorder/Private/Widgets/STakeRecorderCockpit.cpp
--- a/Engine/Plugins/VirtualProduction/Takes/Source/TakeRecorder/Private/Widgets/STakeRecorderCockpit.cpp
+++ b/Engine/Plugins/VirtualProduction/Takes/Source/TakeRecorder/Private/Widgets/STakeRecorderCockpit.cpp
@@ -58,6 +58,7 @@
 #include "ScopedTransaction.h"
 #include "Editor/EditorEngine.h"
 #include "Editor.h"
+#include "Engine/Engine.h"
 #include "Modules/ModuleManager.h"
 #include "LevelEditor.h"
 #include "ISettingsModule.h"
@@ -73,7 +74,7 @@
 // NamingTokens includes
 #include "ITakeRecorderNamingTokensModule.h"
 #include "NamingTokens.h"
-#include "NamingTokensEditorSubsystem.h"
+#include "NamingTokensEngineSubsystem.h"
 #include "NamingTokensSpecifiers.h"
 
 extern UNREALED_API UEditorEngine* GEditor;
@@ -1366,9 +1367,9 @@ TSharedRef<SWidget> STakeRecorderCockpit::OnCreateMenu()
 
 void STakeRecorderCockpit::BindToNamingTokenEvents()
 {
-	if (GIsEditor && !NamingTokensData->TakeRecorderNamingTokens.IsValid())
+	if (GEngine && !NamingTokensData->TakeRecorderNamingTokens.IsValid())
 	{
-		NamingTokensData->TakeRecorderNamingTokens = GEditor->GetEditorSubsystem<UNamingTokensEditorSubsystem>()->GetNamingTokens(ITakeRecorderNamingTokensModule::GetTakeRecorderNamespace());
+		NamingTokensData->TakeRecorderNamingTokens = GEngine->GetEngineSubsystem<UNamingTokensEngineSubsystem>()->GetNamingTokens(ITakeRecorderNamingTokensModule::GetTakeRecorderNamespace());
 		if (ensure(NamingTokensData->TakeRecorderNamingTokens.IsValid()))
 		{
 			const TStrongObjectPtr<UNamingTokens> TakeRecorderNamingTokensPin = NamingTokensData->TakeRecorderNamingTokens.Pin();
diff --git a/Engine/Plugins/VirtualProduction/Takes/Source/TakesCore/Private/TakeMetaData.cpp b/Engine/Plugins/VirtualProduction/Takes/Source/TakesCore/Private/TakeMetaData.cpp
--- a/Engine/Plugins/VirtualProduction/Takes/Source/TakesCore/Private/TakeMetaData.cpp
+++ b/Engine/Plugins/VirtualProduction/Takes/Source/TakesCore/Private/TakeMetaData.cpp
@@ -3,11 +3,12 @@
 #include "TakeMetaData.h"
 
 #include "Editor.h"
+#include "Engine/Engine.h"
 #include "Engine/Level.h"
 #include "ITakeRecorderNamingTokensModule.h"
 #include "NamingTokenData.h"
 #include "NamingTokens/TakeRecorderNamingTokensContext.h"
-#include "NamingTokensEditorSubsystem.h"
+#include "NamingTokensEngineSubsystem.h"
 #include "TakePreset.h"
 #include "TakesCoreBlueprintLibrary.h"
 #include "TakesCoreLog.h"
@@ -132,9 +133,8 @@ FString UTakeMetaData::GenerateAssetPath(const FString& PathFormatString, UTakeR
 
 FNamingTokenResultData UTakeMetaData::ProcessTokens(const FText& InText, UTakeRecorderNamingTokensContext* InContext) const
 {
-	if (!ensure(GIsEditor))
+	if (!ensure(GEngine))
 	{
-		UE_LOG(LogTakesCore, Warning, TEXT("GIsEditor is false during UTakeMetaData::GenerateAssetPath and the generated path will be incorrect."))
 		return FNamingTokenResultData();
 	}
 
@@ -146,7 +146,7 @@ FNamingTokenResultData UTakeMetaData::ProcessTokens(const FText& InText, UTakeRe
 
 	FNamingTokenFilterArgs NamingTokenFilters;
 	NamingTokenFilters.NamespacesToInclude = { ITakeRecorderNamingTokensModule::GetTakeRecorderNamespace() };
-	return GEditor->GetEditorSubsystem<UNamingTokensEditorSubsystem>()->EvaluateTokenText(InText,
+	return GEngine->GetEngineSubsystem<UNamingTokensEngineSubsystem>()->EvaluateTokenText(InText,
 		NamingTokenFilters, { InContext });
 }
 
