NamingTokens: Add missing copyright information.

diff --git a/Engine/Plugins/Developer/NamingTokens/Source/NamingTokensUncookedOnly/NamingTokensUncookedOnly.Build.cs b/Engine/Plugins/Developer/NamingTokens/Source/NamingTokensUncookedOnly/NamingTokensUncookedOnly.Build.cs
--- a/Engine/Plugins/Developer/NamingTokens/Source/NamingTokensUncookedOnly/NamingTokensUncookedOnly.Build.cs
+++ b/Engine/Plugins/Developer/NamingTokens/Source/NamingTokensUncookedOnly/NamingTokensUncookedOnly.Build.cs
@@ -1,4 +1,6 @@
-﻿using UnrealBuildTool;
+﻿// Copyright Epic Games, Inc. All Rights Reserved.
+
+using UnrealBuildTool;
 
 public class NamingTokensUncookedOnly : ModuleRules
 {
