NamingTokens: Change namespace delimiter from "." to ":".

diff --git a/Engine/Plugins/Developer/NamingTokens/Source/NamingTokens/Private/NamingTokens.cpp b/Engine/Plugins/Developer/NamingTokens/Source/NamingTokens/Private/NamingTokens.cpp
--- a/Engine/Plugins/Developer/NamingTokens/Source/NamingTokens/Private/NamingTokens.cpp
+++ b/Engine/Plugins/Developer/NamingTokens/Source/NamingTokens/Private/NamingTokens.cpp
@@ -254,7 +254,7 @@ FText UNamingTokens::ProcessToken(const FNamingTokenData& InToken, FText& InOutF
 		
 	FFormatNamedArguments Args;
 	Args.Add(InToken.TokenKey, ProcessedToken);
-	Args.Add(FString::Printf(TEXT("%s.%s"), *GetNamespace(), *InToken.TokenKey), ProcessedToken);
+	Args.Add(FString::Printf(TEXT("%s%s%s"), *GetNamespace(), *UE::NamingTokens::Utils::GetNamespaceDelimiter(), *InToken.TokenKey), ProcessedToken);
 
 	InOutFormattedText = FText::Format(InOutFormattedText, Args);
 	return ProcessedToken;
diff --git a/Engine/Plugins/Developer/NamingTokens/Source/NamingTokens/Private/NamingTokensEngineSubsystem.cpp b/Engine/Plugins/Developer/NamingTokens/Source/NamingTokens/Private/NamingTokensEngineSubsystem.cpp
--- a/Engine/Plugins/Developer/NamingTokens/Source/NamingTokens/Private/NamingTokensEngineSubsystem.cpp
+++ b/Engine/Plugins/Developer/NamingTokens/Source/NamingTokens/Private/NamingTokensEngineSubsystem.cpp
@@ -310,7 +310,7 @@ TSet<FString> UNamingTokensEngineSubsystem::GetNamingTokenNamespacesFromString(c
 {
 	TSet<FString> Namespaces;
 
-	const FRegexPattern Pattern(R"(\{\s*([a-zA-Z0-9_]+)\.[a-zA-Z0-9_]+\s*\})");
+	const FRegexPattern Pattern(FString::Printf(TEXT(R"(\{\s*([a-zA-Z0-9_]+)%s[a-zA-Z0-9_]+\s*\})"), *UE::NamingTokens::Utils::GetNamespaceDelimiter()));
 	FRegexMatcher Matcher(Pattern, InTokenString);
 	
 	while (Matcher.FindNext())
diff --git a/Engine/Plugins/Developer/NamingTokens/Source/NamingTokens/Private/Utils/NamingTokenUtils.cpp b/Engine/Plugins/Developer/NamingTokens/Source/NamingTokens/Private/Utils/NamingTokenUtils.cpp
--- a/Engine/Plugins/Developer/NamingTokens/Source/NamingTokens/Private/Utils/NamingTokenUtils.cpp
+++ b/Engine/Plugins/Developer/NamingTokens/Source/NamingTokens/Private/Utils/NamingTokenUtils.cpp
@@ -30,11 +30,16 @@ FString UE::NamingTokens::Utils::CreateFormattedToken(const FNamingTokenData& In
 	return FString::Printf(TEXT("{%s}"), *InToken.TokenKey);
 }
 
+FString UE::NamingTokens::Utils::GetNamespaceDelimiter()
+{
+	return TEXT(":");
+}
+
 TArray<FString> UE::NamingTokens::Utils::GetTokenKeysFromString(const FString& InTokenString)
 {
 	TSet<FString, FLocKeySetFuncs> Tokens;
 
-	const FString PatternString = TEXT(R"(\{\s*((?:[a-zA-Z0-9_]+\.)*[a-zA-Z0-9_]+)\s*\})");
+	const FString PatternString = FString::Printf(TEXT(R"(\{\s*((?:[a-zA-Z0-9_]+%s)*[a-zA-Z0-9_]+)\s*\})"), *GetNamespaceDelimiter());
 	const FRegexPattern Pattern(PatternString);
 	FRegexMatcher Matcher(Pattern, InTokenString);
 
@@ -50,7 +55,7 @@ TArray<FString> UE::NamingTokens::Utils::GetTokenKeysFromString(const FString& I
 FString UE::NamingTokens::Utils::GetNamespaceFromTokenKey(const FString& InTokenKey)
 {
 	FString LeftPart, RightPart;
-	InTokenKey.Split(TEXT("."), &LeftPart, &RightPart);
+	InTokenKey.Split(GetNamespaceDelimiter(), &LeftPart, &RightPart);
 
 	return LeftPart;
 }
@@ -59,7 +64,7 @@ FString UE::NamingTokens::Utils::RemoveNamespaceFromTokenKey(const FString& InTo
 {
 	FString Result = InTokenKey;
 	FString LeftPart, RightPart;
-	if (InTokenKey.Split(TEXT("."), &LeftPart, &RightPart))
+	if (InTokenKey.Split(GetNamespaceDelimiter(), &LeftPart, &RightPart))
 	{
 		Result = RightPart;
 	}
diff --git a/Engine/Plugins/Developer/NamingTokens/Source/NamingTokens/Public/Utils/NamingTokenUtils.h b/Engine/Plugins/Developer/NamingTokens/Source/NamingTokens/Public/Utils/NamingTokenUtils.h
--- a/Engine/Plugins/Developer/NamingTokens/Source/NamingTokens/Public/Utils/NamingTokenUtils.h
+++ b/Engine/Plugins/Developer/NamingTokens/Source/NamingTokens/Public/Utils/NamingTokenUtils.h
@@ -15,8 +15,11 @@ namespace UE::NamingTokens::Utils
 	NAMINGTOKENS_API bool ValidateTokenFunction(const UFunction* InFunction);
 
 	/** Formats a token key into {key}. */
-	FString CreateFormattedToken(const FNamingTokenData& InToken);
+	NAMINGTOKENS_API FString CreateFormattedToken(const FNamingTokenData& InToken);
 
+	/** Retrieve the delimiter for the namespace. */
+	NAMINGTOKENS_API FString GetNamespaceDelimiter();
+	
 	/** Extract all token keys from a string. */
 	TArray<FString> GetTokenKeysFromString(const FString& InTokenString);
 
