NamingTokens: Add IgnoreForMemberInitializationTest specifier to FNamingTokensEvaluationData::CurrentDateTime as this is set in the constructor.

diff --git a/Engine/Plugins/Developer/NamingTokens/Source/NamingTokens/Public/NamingTokensEvaluationData.h b/Engine/Plugins/Developer/NamingTokens/Source/NamingTokens/Public/NamingTokensEvaluationData.h
--- a/Engine/Plugins/Developer/NamingTokens/Source/NamingTokens/Public/NamingTokensEvaluationData.h
+++ b/Engine/Plugins/Developer/NamingTokens/Source/NamingTokens/Public/NamingTokensEvaluationData.h
@@ -19,7 +19,7 @@ public:
 	FNamingTokensEvaluationData();
 	
 	/** The current date time. */
-	UPROPERTY(BlueprintReadOnly, Category = "NamingTokens")
+	UPROPERTY(BlueprintReadOnly, Category = "NamingTokens", meta = (IgnoreForMemberInitializationTest))
 	FDateTime CurrentDateTime;
 
 	/** An optional array of context objects for this evaluation. */
