NamingTokens: Adjust tooltip for creating a blueprint graph to inform that a valid Token Key is required.

diff --git a/Engine/Plugins/Developer/NamingTokens/Source/NamingTokensUncookedOnly/Private/Customization/NamingTokensCustomization.cpp b/Engine/Plugins/Developer/NamingTokens/Source/NamingTokensUncookedOnly/Private/Customization/NamingTokensCustomization.cpp
--- a/Engine/Plugins/Developer/NamingTokens/Source/NamingTokensUncookedOnly/Private/Customization/NamingTokensCustomization.cpp
+++ b/Engine/Plugins/Developer/NamingTokens/Source/NamingTokensUncookedOnly/Private/Customization/NamingTokensCustomization.cpp
@@ -127,7 +127,7 @@ void FNamingTokensCustomization::CustomizeChildren(TSharedRef<IPropertyHandle> P
 				.Image(FAppStyle::Get().GetBrush("PListEditor.Button_AddToArray"))
 			]
 			.ButtonStyle(FAppStyle::Get(), "HoverHintOnly")
-			.ToolTipText(LOCTEXT("AddFunction", "Create and assign a new function graph for evaluating tokens."))
+			.ToolTipText(LOCTEXT("AddFunction", "Create and assign a new function graph for evaluating tokens.\nThis requires a valid Token Key entered."))
 			.OnClicked(this, &FNamingTokensCustomization::OnAddFunctionClicked)
 			.IsEnabled(this, &FNamingTokensCustomization::CanAddFunction)
 		]
