NamingTokens: Allow case insensitivity when there is no exact match.

diff --git a/Engine/Plugins/Developer/NamingTokens/Source/NamingTokens/Private/NamingTokenData.cpp b/Engine/Plugins/Developer/NamingTokens/Source/NamingTokens/Private/NamingTokenData.cpp
--- a/Engine/Plugins/Developer/NamingTokens/Source/NamingTokens/Private/NamingTokenData.cpp
+++ b/Engine/Plugins/Developer/NamingTokens/Source/NamingTokens/Private/NamingTokenData.cpp
@@ -21,7 +21,12 @@ FNamingTokenData::FNamingTokenData(const FString& InTokenKey, const FText& InTok
 	Description = InTokenDescription;
 }
 
+bool FNamingTokenData::Equals(const FNamingTokenData& Other, bool bCaseSensitive) const
+{
+	return TokenKey.Equals(Other.TokenKey, bCaseSensitive ? ESearchCase::CaseSensitive : ESearchCase::IgnoreCase);
+}
+
 bool FNamingTokenData::operator==(const FNamingTokenData& Other) const
 {
-	return TokenKey.Equals(Other.TokenKey, ESearchCase::CaseSensitive);
+	return Equals(Other);
 }
diff --git a/Engine/Plugins/Developer/NamingTokens/Source/NamingTokens/Private/NamingTokens.cpp b/Engine/Plugins/Developer/NamingTokens/Source/NamingTokens/Private/NamingTokens.cpp
--- a/Engine/Plugins/Developer/NamingTokens/Source/NamingTokens/Private/NamingTokens.cpp
+++ b/Engine/Plugins/Developer/NamingTokens/Source/NamingTokens/Private/NamingTokens.cpp
@@ -87,35 +87,76 @@ FNamingTokenResultData UNamingTokens::EvaluateTokenText(const FText& InTokenText
 				continue;
 			}
 			CheckedKeys.Add(TokenKey);
-			
-			FNamingTokenValueData ValueData;
-			
-			FText TokenValue;
+
 			const FNamingTokenData* SelectedToken = nullptr;
-			if (FNamingTokenData* DefaultToken = DefaultTokensSet.Find(FNamingTokenData(TokenKey)))
-			{
-				SelectedToken = DefaultToken;
-			}
-			else if (FNamingTokenData* CustomToken = CustomTokenSet.Find(FNamingTokenData(TokenKey)))
+			
+			FNamingTokenData NamingTokenKeyData(TokenKey);
+
+			// Check case-sensitive match.
 			{
-				SelectedToken = CustomToken;
+				if (FNamingTokenData* DefaultToken = DefaultTokensSet.Find(NamingTokenKeyData))
+				{
+					SelectedToken = DefaultToken;
+				}
+				else if (FNamingTokenData* CustomToken = CustomTokenSet.Find(NamingTokenKeyData))
+				{
+					SelectedToken = CustomToken;
+				}
+				else
+				{
+					// Not found in a class definition, check external tokens.
+					for (const TTuple<FGuid, TArray<FNamingTokenData>>& ExternalTokenKeyPair : ExternalTokens)
+					{
+						if (const FNamingTokenData* ExternalToken = ExternalTokenKeyPair.Value.FindByKey(NamingTokenKeyData))
+						{
+							SelectedToken = ExternalToken;
+							break;
+						}
+					}
+				}
 			}
-			else
+
+			// No exact case-sensitive match found, check case-insensitive.
+			if (!SelectedToken && !InEvaluationData.bForceCaseSensitive)
 			{
-				// Not found in a class definition, check external tokens.
-				for (const TTuple<FGuid, TArray<FNamingTokenData>>& ExternalTokenKeyPair : ExternalTokens)
+				constexpr bool bCaseSensitive = false;
+				if (FNamingTokenData* DefaultToken = DefaultTokens.FindByPredicate([&NamingTokenKeyData](const FNamingTokenData& NamingToken)
+				{
+					return NamingToken.Equals(NamingTokenKeyData, bCaseSensitive);
+				}))
+				{
+					SelectedToken = DefaultToken;
+				}
+				else if (FNamingTokenData* CustomToken = CustomTokens.FindByPredicate([&NamingTokenKeyData](const FNamingTokenData& NamingToken)
+				{
+					return NamingToken.Equals(NamingTokenKeyData, bCaseSensitive);
+				}))
+				{
+					SelectedToken = CustomToken;
+				}
+				else
 				{
-					if (const FNamingTokenData* ExternalToken = ExternalTokenKeyPair.Value.FindByKey(FNamingTokenData(TokenKey)))
+					for (const TTuple<FGuid, TArray<FNamingTokenData>>& ExternalTokenKeyPair : ExternalTokens)
 					{
-						SelectedToken = ExternalToken;
-						break;
+						if (const FNamingTokenData* ExternalToken =
+							ExternalTokenKeyPair.Value.FindByPredicate([&NamingTokenKeyData](const FNamingTokenData& ExternalTokenKeyPairValue)
+						{
+							return NamingTokenKeyData.Equals(ExternalTokenKeyPairValue, bCaseSensitive);
+						}))
+						{
+							SelectedToken = ExternalToken;
+							break;
+						}
 					}
 				}
 			}
+
+			FNamingTokenValueData ValueData;
+			FText TokenValue;
 			
 			if (SelectedToken)
 			{
-				TokenValue = ProcessToken(*SelectedToken, FormattedText);
+				TokenValue = ProcessToken(*SelectedToken, TokenKey, FormattedText);
 				ValueData.bWasEvaluated = true;
 			}
 			else
@@ -123,7 +164,7 @@ FNamingTokenResultData UNamingTokens::EvaluateTokenText(const FText& InTokenText
 				ValueData.bWasEvaluated = false;
 			}
 			
-			ValueData.TokenKey = TokenKey;
+			ValueData.TokenKey = SelectedToken ? SelectedToken->TokenKey : TokenKey;
 			ValueData.TokenValue = TokenValue;
 			ValueData.ProvidedNamespace = TokenNamespace;
 			Result.TokenValues.Add(MoveTemp(ValueData));
@@ -227,7 +268,7 @@ void UNamingTokens::OnPostEvaluate_Implementation()
 {
 }
 
-FText UNamingTokens::ProcessToken(const FNamingTokenData& InToken, FText& InOutFormattedText)
+FText UNamingTokens::ProcessToken(const FNamingTokenData& InToken, const FString& InUserTokenString, FText& InOutFormattedText)
 {
 	const FString FormattedToken = *UE::NamingTokens::Utils::CreateFormattedToken(InToken);
 
@@ -253,8 +294,8 @@ FText UNamingTokens::ProcessToken(const FNamingTokenData& InToken, FText& InOutF
 	}
 		
 	FFormatNamedArguments Args;
-	Args.Add(InToken.TokenKey, ProcessedToken);
-	Args.Add(FString::Printf(TEXT("%s%s%s"), *GetNamespace(), *UE::NamingTokens::Utils::GetNamespaceDelimiter(), *InToken.TokenKey), ProcessedToken);
+	Args.Add(InUserTokenString, ProcessedToken);
+	Args.Add(FString::Printf(TEXT("%s%s%s"), *GetNamespace(), *UE::NamingTokens::Utils::GetNamespaceDelimiter(), *InUserTokenString), ProcessedToken);
 
 	InOutFormattedText = FText::Format(InOutFormattedText, Args);
 	return ProcessedToken;
diff --git a/Engine/Plugins/Developer/NamingTokens/Source/NamingTokens/Private/NamingTokensEngineSubsystem.cpp b/Engine/Plugins/Developer/NamingTokens/Source/NamingTokens/Private/NamingTokensEngineSubsystem.cpp
--- a/Engine/Plugins/Developer/NamingTokens/Source/NamingTokens/Private/NamingTokensEngineSubsystem.cpp
+++ b/Engine/Plugins/Developer/NamingTokens/Source/NamingTokens/Private/NamingTokensEngineSubsystem.cpp
@@ -213,6 +213,7 @@ FNamingTokenResultData UNamingTokensEngineSubsystem::EvaluateTokenText(const FTe
 
 	FNamingTokensEvaluationData EvaluationData;
 	EvaluationData.Contexts = InContexts;
+	EvaluationData.bForceCaseSensitive = InFilter.bForceCaseSensitive;
 	for (const FString& Namespace : Namespaces)
 	{
 		if (UNamingTokens* NamingTokens = GetNamingTokens(Namespace))
diff --git a/Engine/Plugins/Developer/NamingTokens/Source/NamingTokens/Public/NamingTokenData.h b/Engine/Plugins/Developer/NamingTokens/Source/NamingTokens/Public/NamingTokenData.h
--- a/Engine/Plugins/Developer/NamingTokens/Source/NamingTokens/Public/NamingTokenData.h
+++ b/Engine/Plugins/Developer/NamingTokens/Source/NamingTokens/Public/NamingTokenData.h
@@ -83,6 +83,9 @@ public:
 
 	UE_API bool operator==(const FNamingTokenData& Other) const;
 
+	/** Checks for equality dependent on case sensitivity. */
+	UE_API bool Equals(const FNamingTokenData& Other, bool bCaseSensitive = true) const;
+	
 	friend FORCEINLINE uint32 GetTypeHash(const FNamingTokenData& Item)
 	{
 		const uint32 Hash = FCrc::StrCrc32(*Item.TokenKey);
diff --git a/Engine/Plugins/Developer/NamingTokens/Source/NamingTokens/Public/NamingTokens.h b/Engine/Plugins/Developer/NamingTokens/Source/NamingTokens/Public/NamingTokens.h
--- a/Engine/Plugins/Developer/NamingTokens/Source/NamingTokens/Public/NamingTokens.h
+++ b/Engine/Plugins/Developer/NamingTokens/Source/NamingTokens/Public/NamingTokens.h
@@ -149,7 +149,7 @@ private:
 	FOnPostEvaluate OnPostEvaluateEvent;
 	
 	/** Process a token if is defined and return evaluated text. */
-	UE_API FText ProcessToken(const FNamingTokenData& InToken, FText& InOutFormattedText);
+	UE_API FText ProcessToken(const FNamingTokenData& InToken, const FString& InUserTokenString, FText& InOutFormattedText);
 	
 	/** Return a blueprint function for a token processor if it exists. */
 	UE_API UFunction* FindBlueprintFunctionForToken(const FNamingTokenData& InTokenData) const;
diff --git a/Engine/Plugins/Developer/NamingTokens/Source/NamingTokens/Public/NamingTokensEngineSubsystem.h b/Engine/Plugins/Developer/NamingTokens/Source/NamingTokens/Public/NamingTokensEngineSubsystem.h
--- a/Engine/Plugins/Developer/NamingTokens/Source/NamingTokens/Public/NamingTokensEngineSubsystem.h
+++ b/Engine/Plugins/Developer/NamingTokens/Source/NamingTokens/Public/NamingTokensEngineSubsystem.h
@@ -29,6 +29,10 @@ struct FNamingTokenFilterArgs
 	/** Include global namespaces. */
 	UPROPERTY(EditAnywhere, BlueprintReadWrite, Category=NamingTokens)
 	bool bIncludeGlobal = true;
+
+	/** When false, we fall back to case-insensitive if an exact match isn't found. */
+	UPROPERTY(EditAnywhere, BlueprintReadWrite, Category=NamingTokens)
+	bool bForceCaseSensitive = false;
 };
 
 /**
diff --git a/Engine/Plugins/Developer/NamingTokens/Source/NamingTokens/Public/NamingTokensEvaluationData.h b/Engine/Plugins/Developer/NamingTokens/Source/NamingTokens/Public/NamingTokensEvaluationData.h
--- a/Engine/Plugins/Developer/NamingTokens/Source/NamingTokens/Public/NamingTokensEvaluationData.h
+++ b/Engine/Plugins/Developer/NamingTokens/Source/NamingTokens/Public/NamingTokensEvaluationData.h
@@ -25,4 +25,7 @@ public:
 	/** An optional array of context objects for this evaluation. */
 	UPROPERTY(BlueprintReadOnly, Category = "NamingTokens")
 	TArray<TObjectPtr<UObject>> Contexts;
+
+	/** If case sensitivity should be enforced. */
+	bool bForceCaseSensitive = false;
 };
\ No newline at end of file
