NamingTokens: Add validation to token keys and namespaces. Namespaces and keys should only contain alphanumeric characters, underscores, and cannot be empty. Visual warnings are displayed in blueprints and error logs are present when running Naming Tokens. Global namespaces cannot be registered if they are invalid.

Fix not being able to undo creating a new token graph.

diff --git a/Engine/Plugins/Developer/NamingTokens/Source/NamingTokens/Private/NamingTokens.cpp b/Engine/Plugins/Developer/NamingTokens/Source/NamingTokens/Private/NamingTokens.cpp
--- a/Engine/Plugins/Developer/NamingTokens/Source/NamingTokens/Private/NamingTokens.cpp
+++ b/Engine/Plugins/Developer/NamingTokens/Source/NamingTokens/Private/NamingTokens.cpp
@@ -23,6 +23,31 @@ void UNamingTokens::PostInitProperties()
 	CreateDefaultTokens();
 }
 
+void UNamingTokens::Validate() const
+{
+	// Namespace validation
+	{
+		FText ErrorMessage;
+		if (!UE::NamingTokens::Utils::ValidateName(Namespace, ErrorMessage))
+		{
+			UE_LOG(LogNamingTokens, Error, TEXT("NamingTokens Namespace '%s' has an invalid name. This should be corrected by the owner of this namespace. Error: %s"),
+				*Namespace, *ErrorMessage.ToString());
+		}
+	}
+
+	// Individual token validation
+	const TArray<FNamingTokenData> AllTokens = GetAllTokens();
+	for (const FNamingTokenData& TokenData : AllTokens)
+	{
+		FText ErrorMessage;
+		if (!UE::NamingTokens::Utils::ValidateName(TokenData.TokenKey, ErrorMessage))
+		{
+			UE_LOG(LogNamingTokens, Error, TEXT("NamingTokens Token '%s' under namespace '%s' has an invalid name. This should be corrected by the owner of this namespace. Error: %s"),
+				*TokenData.TokenKey, *Namespace, *ErrorMessage.ToString());
+		}
+	}
+}
+
 void UNamingTokens::CreateDefaultTokens()
 {
 	DefaultTokens.Empty();
diff --git a/Engine/Plugins/Developer/NamingTokens/Source/NamingTokens/Private/NamingTokensEngineSubsystem.cpp b/Engine/Plugins/Developer/NamingTokens/Source/NamingTokens/Private/NamingTokensEngineSubsystem.cpp
--- a/Engine/Plugins/Developer/NamingTokens/Source/NamingTokens/Private/NamingTokensEngineSubsystem.cpp
+++ b/Engine/Plugins/Developer/NamingTokens/Source/NamingTokens/Private/NamingTokensEngineSubsystem.cpp
@@ -56,6 +56,8 @@ UNamingTokens* UNamingTokensEngineSubsystem::GetNamingTokens(const FString& InNa
 		if (const UClass* Class = InTokensClass.LoadSynchronous())
 		{
 			UNamingTokens* NamingTokens = NewObject<UNamingTokens>(GetTransientPackage(), Class);
+			// Initial validation on load. This will log any errors.
+			NamingTokens->Validate();
 			if (ensure(NamingTokens->GetNamespace() == InNamespace))
 			{
 				CachedNamingTokens.Add(InNamespace, NamingTokens);
@@ -279,7 +281,16 @@ void UNamingTokensEngineSubsystem::RegisterGlobalNamespace(const FString& InName
 			 *InNamespace, *UnregisterFunctionName);
 		return;
 	}
-	GlobalNamespaces.Add(InNamespace);
+	FText ErrorMessage;
+	if (UE::NamingTokens::Utils::ValidateName(InNamespace, ErrorMessage))
+	{
+		GlobalNamespaces.Add(InNamespace);
+	}
+	else
+	{
+		UE_LOG(LogNamingTokens, Error, TEXT("NamingTokens Namespace '%s' cannot be registered. Error: %s"),
+			*InNamespace, *ErrorMessage.ToString());
+	}
 }
 
 void UNamingTokensEngineSubsystem::UnregisterGlobalNamespace(const FString& InNamespace)
diff --git a/Engine/Plugins/Developer/NamingTokens/Source/NamingTokens/Private/Utils/NamingTokenUtils.cpp b/Engine/Plugins/Developer/NamingTokens/Source/NamingTokens/Private/Utils/NamingTokenUtils.cpp
--- a/Engine/Plugins/Developer/NamingTokens/Source/NamingTokens/Private/Utils/NamingTokenUtils.cpp
+++ b/Engine/Plugins/Developer/NamingTokens/Source/NamingTokens/Private/Utils/NamingTokenUtils.cpp
@@ -97,3 +97,38 @@ FString UE::NamingTokens::Utils::RemoveNamespaceFromTokenKey(const FString& InTo
 	}
 	return Result;
 }
+
+bool UE::NamingTokens::Utils::ValidateName(const FString& InName, FText& OutErrorMessage)
+{
+	OutErrorMessage = FText::GetEmpty();
+
+	if (InName.IsEmpty())
+	{
+		OutErrorMessage = NSLOCTEXT("NamingTokenUtils", "StringEmpty", "String is empty");
+	}
+	else
+	{
+		TSet<TCHAR> UniqueInvalidChars;
+		for (const TCHAR& Char : InName)
+		{
+			// We may want to eventually use FName::IsValidXName with INVALID_OBJECTNAME_CHARACTERS or variant.
+			// For right now, this explicitly allows the 3 types we want and reports invalid characters.
+			if (!FChar::IsAlpha(Char) && !FChar::IsDigit(Char) && Char != TEXT('_'))
+			{
+				UniqueInvalidChars.Add(Char);
+			}
+		}
+
+		if (UniqueInvalidChars.Num() > 0)
+		{
+			TArray<FString> InvalidStrings;
+			for (const TCHAR& Char : UniqueInvalidChars)
+			{
+				InvalidStrings.Add(FString::Printf(TEXT("'%c'"), Char));
+			}
+			OutErrorMessage = FText::FromString(TEXT("Invalid characters: ") + FString::Join(InvalidStrings, TEXT(", ")));
+		}
+	}
+	
+	return OutErrorMessage.IsEmpty();
+}
diff --git a/Engine/Plugins/Developer/NamingTokens/Source/NamingTokens/Public/NamingTokenData.h b/Engine/Plugins/Developer/NamingTokens/Source/NamingTokens/Public/NamingTokenData.h
--- a/Engine/Plugins/Developer/NamingTokens/Source/NamingTokens/Public/NamingTokenData.h
+++ b/Engine/Plugins/Developer/NamingTokens/Source/NamingTokens/Public/NamingTokenData.h
@@ -62,7 +62,12 @@ public:
 	UE_API FNamingTokenData(const FString& InTokenKey, const FText& InTokenDisplayName,
 	                 const FText& InTokenDescription, const FTokenProcessorDelegateNative& InTokenProcessor);
 
-	/** The key of the token to use. This is what the text must match in order to be evaluated. Brackets are automatically added and do not need to be included. */
+	/**
+	 * The key of the token to use.
+	 * This is what the text must match in order to be evaluated. Brackets are automatically added and do not need to be included.
+	 *
+	 * @note Must contain alphanumeric and '_' characters only and cannot be empty.
+	 */
 	UPROPERTY(BlueprintReadWrite, EditAnywhere, Category = "NamingTokens")
 	FString TokenKey;
 
diff --git a/Engine/Plugins/Developer/NamingTokens/Source/NamingTokens/Public/NamingTokens.h b/Engine/Plugins/Developer/NamingTokens/Source/NamingTokens/Public/NamingTokens.h
--- a/Engine/Plugins/Developer/NamingTokens/Source/NamingTokens/Public/NamingTokens.h
+++ b/Engine/Plugins/Developer/NamingTokens/Source/NamingTokens/Public/NamingTokens.h
@@ -23,6 +23,9 @@ public:
 
 	UE_API virtual void PostInitProperties() override;
 
+	/** Validate internal values including namespace and token keys. */
+	UE_API void Validate() const;
+	
 	/** Create any default tokens. */
 	UE_API void CreateDefaultTokens();
 
@@ -170,7 +173,11 @@ protected:
 	UPROPERTY(BlueprintReadOnly, Category = "NamingTokens")
 	FNamingTokensEvaluationData CurrentEvaluationData;
 	
-	/** The namespace to identify this token. */
+	/**
+	 * The namespace to identify this token.
+	 * 
+	 * @note Must contain alphanumeric and '_' characters only and cannot be empty.
+	 */
 	UPROPERTY(EditDefaultsOnly, AssetRegistrySearchable, Category = "NamingTokens")
 	FString Namespace;
 	
diff --git a/Engine/Plugins/Developer/NamingTokens/Source/NamingTokens/Public/Utils/NamingTokenUtils.h b/Engine/Plugins/Developer/NamingTokens/Source/NamingTokens/Public/Utils/NamingTokenUtils.h
--- a/Engine/Plugins/Developer/NamingTokens/Source/NamingTokens/Public/Utils/NamingTokenUtils.h
+++ b/Engine/Plugins/Developer/NamingTokens/Source/NamingTokens/Public/Utils/NamingTokenUtils.h
@@ -39,4 +39,7 @@ namespace UE::NamingTokens::Utils
 	 * @return The token key without the namespace, or the InTokenKey if there was no namespace.
 	 */
 	NAMINGTOKENS_API FString RemoveNamespaceFromTokenKey(const FString& InTokenKey);
+
+	/** Validates a name used for a token key or namespace. */
+	NAMINGTOKENS_API bool ValidateName(const FString& InName, FText& OutErrorMessage);
 }
\ No newline at end of file
diff --git a/Engine/Plugins/Developer/NamingTokens/Source/NamingTokensUncookedOnly/Private/Customization/NamingTokensCustomization.cpp b/Engine/Plugins/Developer/NamingTokens/Source/NamingTokensUncookedOnly/Private/Customization/NamingTokensCustomization.cpp
--- a/Engine/Plugins/Developer/NamingTokens/Source/NamingTokensUncookedOnly/Private/Customization/NamingTokensCustomization.cpp
+++ b/Engine/Plugins/Developer/NamingTokens/Source/NamingTokensUncookedOnly/Private/Customization/NamingTokensCustomization.cpp
@@ -4,25 +4,111 @@
 
 #include "NamingTokenData.h"
 #include "NamingTokens.h"
-#include "Utils/NamingTokenUtils.h"
 #include "Utils/NamingTokensEditorUtils.h"
+#include "Utils/NamingTokenUtils.h"
 
+#include "DetailLayoutBuilder.h"
 #include "DetailWidgetRow.h"
 #include "Engine/Blueprint.h"
 #include "IDetailChildrenBuilder.h"
 #include "IDetailPropertyRow.h"
 #include "IPropertyUtilities.h"
+#include "ScopedTransaction.h"
 #include "Widgets/Input/SButton.h"
 #include "Widgets/Input/SComboBox.h"
 #include "Widgets/Text/STextBlock.h"
 
 #define LOCTEXT_NAMESPACE "NamingTokensCustomization"
 
-void FNamingTokensCustomization::CustomizeHeader(TSharedRef<IPropertyHandle> PropertyHandle, FDetailWidgetRow& HeaderRow,
-                                                 IPropertyTypeCustomizationUtils& CustomizationUtils)
+namespace UE::NamingTokens::Customization::Private
 {
-	CustomizationUtilsPtr = &CustomizationUtils;
+	// If the error message should be visible.
+	EVisibility GetErrorVisibilityFromProperty(const TSharedPtr<IPropertyHandle>& InPropertyHandle, const TSharedPtr<FText>& InErrorMessage)
+	{
+		if (InPropertyHandle.IsValid() && InErrorMessage.IsValid())
+		{
+			FString NamespaceValue;
+			if (InPropertyHandle->GetValue(NamespaceValue) == FPropertyAccess::Result::Success
+				&& Utils::ValidateName(NamespaceValue, *InErrorMessage))
+			{
+				return EVisibility::Collapsed;
+			}
+		}
+
+		return EVisibility::Visible;
+	}
 
+	// Construct the error tooltip message.
+	FText CreateErrorTooltipMessage(const TSharedPtr<FText>& InErrorMessage)
+	{
+		if (InErrorMessage.IsValid() && !InErrorMessage->IsEmpty())
+		{
+			FFormatNamedArguments Args;
+			Args.Add(TEXT("Error"), *InErrorMessage);
+			return FText::Format(LOCTEXT("ValueError", "Error: {Error}. Alphanumeric and '_' characters are allowed."), Args);
+		}
+		return FText::GetEmpty();
+	}
+	
+	// Creates the standard row widget for a property, adding in an error icon for validation.
+	void CreateRowWidgetWithError(IDetailPropertyRow& InRow, const TSharedPtr<IPropertyHandle>& InPropertyHandle,
+		const TSharedPtr<FText>& InErrorMessage)
+	{
+		InRow.CustomWidget()
+		.NameContent()
+		[
+			InPropertyHandle->CreatePropertyNameWidget()
+		]
+		.ValueContent()
+		.MaxDesiredWidth(250.0f)
+		[
+			SNew(SHorizontalBox)
+			+ SHorizontalBox::Slot()
+			.FillWidth(1.0f)
+			[
+				InPropertyHandle->CreatePropertyValueWidget()
+			]
+			+ SHorizontalBox::Slot()
+			.AutoWidth()
+			.VAlign(VAlign_Center)
+			.HAlign(HAlign_Center)
+			.Padding(2.f, 0.f, 0.f, 0.f)
+			[
+				SNew(SImage)
+				.Image(FAppStyle::Get().GetBrush("Icons.Error"))
+				.ColorAndOpacity(FAppStyle::Get().GetSlateColor("Colors.AccentRed"))
+				.ToolTipText_Lambda([InErrorMessage]()
+				{
+					return CreateErrorTooltipMessage(InErrorMessage);
+				})
+				.Visibility_Lambda([InPropertyHandle, InErrorMessage]()
+				{
+					return GetErrorVisibilityFromProperty(InPropertyHandle, InErrorMessage);
+				})
+			]
+		];
+	}
+}
+
+void FNamingTokensCustomization::CustomizeDetails(IDetailLayoutBuilder& DetailBuilder)
+{
+	NamespacePropertyHandle = DetailBuilder.GetProperty(UNamingTokens::GetNamespacePropertyName());
+	TokenKeyErrorMessage = MakeShared<FText>();
+	if (NamespacePropertyHandle.IsValid())
+	{
+		if (IDetailPropertyRow* Row = DetailBuilder.EditDefaultProperty(NamespacePropertyHandle))
+		{
+			UE::NamingTokens::Customization::Private::CreateRowWidgetWithError(*Row, NamespacePropertyHandle, TokenKeyErrorMessage);
+		}
+	}
+}
+
+void FNamingTokensDataCustomization::CustomizeHeader(TSharedRef<IPropertyHandle> PropertyHandle, FDetailWidgetRow& HeaderRow,
+                                                     IPropertyTypeCustomizationUtils& CustomizationUtils)
+{
+	CustomizationUtilsPtr = &CustomizationUtils;
+	NamespaceErrorMessage = MakeShared<FText>();
+	
 	if (const UNamingTokens* OwningTokens = GetOwningNamingTokens())
 	{
 		OwningBlueprint = UBlueprint::GetBlueprintFromClass(OwningTokens->GetClass());
@@ -36,7 +122,7 @@ void FNamingTokensCustomization::CustomizeHeader(TSharedRef<IPropertyHandle> Pro
 		}
 	}
 
-	// If owning blueprint is null, then we may not be customizing the naming tokens directly, such as a settings object.
+	// If owning blueprint is null, then we may not be customizing the naming tokens directly, such as a settings object,
 	// or the tokens are predefined and can't be extended in BP.
 	
 	HeaderRow
@@ -50,7 +136,7 @@ void FNamingTokensCustomization::CustomizeHeader(TSharedRef<IPropertyHandle> Pro
 	];
 }
 
-void FNamingTokensCustomization::CustomizeChildren(TSharedRef<IPropertyHandle> PropertyHandle, IDetailChildrenBuilder& ChildBuilder,
+void FNamingTokensDataCustomization::CustomizeChildren(TSharedRef<IPropertyHandle> PropertyHandle, IDetailChildrenBuilder& ChildBuilder,
 	IPropertyTypeCustomizationUtils& CustomizationUtils)
 {
 	const bool bShouldApplyCustomization = OwningBlueprint.IsValid();
@@ -89,7 +175,12 @@ void FNamingTokensCustomization::CustomizeChildren(TSharedRef<IPropertyHandle> P
 		{
 			continue;
 		}
-		ChildBuilder.AddProperty(ChildHandle.ToSharedRef());
+		IDetailPropertyRow& Row = ChildBuilder.AddProperty(ChildHandle.ToSharedRef());
+
+		if (ChildHandle->GetProperty() == TokenKeyHandle->GetProperty())
+		{
+			UE::NamingTokens::Customization::Private::CreateRowWidgetWithError(Row, TokenKeyHandle, NamespaceErrorMessage);
+		}
 	}
 
 	if (!bShouldApplyCustomization)
@@ -109,12 +200,12 @@ void FNamingTokensCustomization::CustomizeChildren(TSharedRef<IPropertyHandle> P
 		[
 			SNew(SComboBox<TSharedPtr<FString>>)
 			.OptionsSource(&FunctionNames)
-			.OnGenerateWidget(this, &FNamingTokensCustomization::MakeComboBoxWidget)
-			.OnSelectionChanged(this, &FNamingTokensCustomization::OnFunctionSelected, FunctionNameHandle)
+			.OnGenerateWidget(this, &FNamingTokensDataCustomization::MakeComboBoxWidget)
+			.OnSelectionChanged(this, &FNamingTokensDataCustomization::OnFunctionSelected, FunctionNameHandle)
 			.InitiallySelectedItem(SelectedFunctionName)
 			[
 				SNew(STextBlock)
-				.Text(this, &FNamingTokensCustomization::GetSelectedFunctionText)
+				.Text(this, &FNamingTokensDataCustomization::GetSelectedFunctionText)
 			]
 		]
 		+ SHorizontalBox::Slot()
@@ -128,13 +219,13 @@ void FNamingTokensCustomization::CustomizeChildren(TSharedRef<IPropertyHandle> P
 			]
 			.ButtonStyle(FAppStyle::Get(), "HoverHintOnly")
 			.ToolTipText(LOCTEXT("AddFunction", "Create and assign a new function graph for evaluating tokens.\nThis requires a valid Token Key entered."))
-			.OnClicked(this, &FNamingTokensCustomization::OnAddFunctionClicked)
-			.IsEnabled(this, &FNamingTokensCustomization::CanAddFunction)
+			.OnClicked(this, &FNamingTokensDataCustomization::OnAddFunctionClicked)
+			.IsEnabled(this, &FNamingTokensDataCustomization::CanAddFunction)
 		]
 	];
 }
 
-UNamingTokens* FNamingTokensCustomization::GetOwningNamingTokens() const
+UNamingTokens* FNamingTokensDataCustomization::GetOwningNamingTokens() const
 {
 	if (CustomizationUtilsPtr)
 	{
@@ -155,7 +246,7 @@ UNamingTokens* FNamingTokensCustomization::GetOwningNamingTokens() const
 	return nullptr;
 }
 
-TArray<UFunction*> FNamingTokensCustomization::GetAvailableFunctions() const
+TArray<UFunction*> FNamingTokensDataCustomization::GetAvailableFunctions() const
 {
 	TArray<UFunction*> Result;
 	if (OwningBlueprint.IsValid())
@@ -174,13 +265,13 @@ TArray<UFunction*> FNamingTokensCustomization::GetAvailableFunctions() const
 	return Result;
 }
 
-TSharedRef<SWidget> FNamingTokensCustomization::MakeComboBoxWidget(TSharedPtr<FString> InItem)
+TSharedRef<SWidget> FNamingTokensDataCustomization::MakeComboBoxWidget(TSharedPtr<FString> InItem)
 {
 	return SNew(STextBlock)
 		.Text(FText::FromString(*InItem));
 }
 
-void FNamingTokensCustomization::OnFunctionSelected(TSharedPtr<FString> NewValue, ESelectInfo::Type SelectInfo,
+void FNamingTokensDataCustomization::OnFunctionSelected(TSharedPtr<FString> NewValue, ESelectInfo::Type SelectInfo,
 	TSharedPtr<IPropertyHandle> InFunctionNameHandle)
 {
 	if (NewValue.IsValid())
@@ -190,17 +281,18 @@ void FNamingTokensCustomization::OnFunctionSelected(TSharedPtr<FString> NewValue
 	}
 }
 
-FText FNamingTokensCustomization::GetSelectedFunctionText() const
+FText FNamingTokensDataCustomization::GetSelectedFunctionText() const
 {
 	return SelectedFunctionName.IsValid() ? FText::FromString(*SelectedFunctionName) : FText::FromString(TEXT("Select Function"));
 }
 
-FReply FNamingTokensCustomization::OnAddFunctionClicked()
+FReply FNamingTokensDataCustomization::OnAddFunctionClicked()
 {
 	if (UBlueprint* Blueprint = OwningBlueprint.Get())
 	{
 		FString TokenKey;
 		TokenKeyHandle->GetValue(TokenKey);
+		FScopedTransaction Transaction(LOCTEXT("CreateTokenGraph", "Create Token Graph"));
 
 		const FName NewFunctionName = UE::NamingTokens::Utils::Editor::Private::CreateNewTokenGraph(Blueprint, TokenKey);
 		FunctionNameHandle->SetValue(NewFunctionName);
@@ -209,13 +301,19 @@ FReply FNamingTokensCustomization::OnAddFunctionClicked()
 	return FReply::Handled();
 }
 
-bool FNamingTokensCustomization::CanAddFunction() const
+bool FNamingTokensDataCustomization::CanAddFunction() const
 {
 	FString TokenKey;
 	if (TokenKeyHandle.IsValid())
 	{
 		TokenKeyHandle->GetValue(TokenKey);
-		
+
+		FText ErrorMessage;
+		if (!UE::NamingTokens::Utils::ValidateName(TokenKey, ErrorMessage))
+		{
+			return false;
+		}
+			
 		const FString BaseFunctionName = UE::NamingTokens::Utils::Editor::Private::CreateBaseTokenFunctionName(TokenKey);
 
 		// Make sure the name isn't already being used.
diff --git a/Engine/Plugins/Developer/NamingTokens/Source/NamingTokensUncookedOnly/Private/Customization/NamingTokensCustomization.h b/Engine/Plugins/Developer/NamingTokens/Source/NamingTokensUncookedOnly/Private/Customization/NamingTokensCustomization.h
--- a/Engine/Plugins/Developer/NamingTokens/Source/NamingTokensUncookedOnly/Private/Customization/NamingTokensCustomization.h
+++ b/Engine/Plugins/Developer/NamingTokens/Source/NamingTokensUncookedOnly/Private/Customization/NamingTokensCustomization.h
@@ -2,19 +2,38 @@
 
 #pragma once
 
-#include "IPropertyTypeCustomization.h"
 #include "CoreMinimal.h"
+#include "IDetailCustomization.h"
 #include "Input/Reply.h"
+#include "IPropertyTypeCustomization.h"
+#include "Layout/Visibility.h"
 
 class UNamingTokens;
 class SWidget;
 
-class FNamingTokensCustomization final : public IPropertyTypeCustomization
+class FNamingTokensCustomization final : public IDetailCustomization
+{
+public:
+	static TSharedRef<IDetailCustomization> MakeInstance()
+	{
+		return MakeShareable(new FNamingTokensCustomization);
+	}
+
+	virtual void CustomizeDetails(IDetailLayoutBuilder& DetailBuilder) override;
+	
+private:
+	/** Handle to namespace property. */
+	TSharedPtr<IPropertyHandle> NamespacePropertyHandle;
+	/** Stores error text. */
+	TSharedPtr<FText> TokenKeyErrorMessage;
+};
+
+class FNamingTokensDataCustomization final : public IPropertyTypeCustomization
 {
 public:
 	static TSharedRef<IPropertyTypeCustomization> MakeInstance()
 	{
-		return MakeShareable(new FNamingTokensCustomization());
+		return MakeShareable(new FNamingTokensDataCustomization());
 	}
 
 	virtual void CustomizeHeader(TSharedRef<IPropertyHandle> PropertyHandle, FDetailWidgetRow& HeaderRow, IPropertyTypeCustomizationUtils& CustomizationUtils) override;
@@ -54,4 +73,6 @@ private:
 	TSharedPtr<IPropertyHandle> FunctionNameHandle;
 	/** Property handle to the token key. */
 	TSharedPtr<IPropertyHandle> TokenKeyHandle;
+	/** Stores error text. */
+	TSharedPtr<FText> NamespaceErrorMessage;
 };
\ No newline at end of file
diff --git a/Engine/Plugins/Developer/NamingTokens/Source/NamingTokensUncookedOnly/Private/NamingTokensUncookedOnlyModule.cpp b/Engine/Plugins/Developer/NamingTokens/Source/NamingTokensUncookedOnly/Private/NamingTokensUncookedOnlyModule.cpp
--- a/Engine/Plugins/Developer/NamingTokens/Source/NamingTokensUncookedOnly/Private/NamingTokensUncookedOnlyModule.cpp
+++ b/Engine/Plugins/Developer/NamingTokens/Source/NamingTokensUncookedOnly/Private/NamingTokensUncookedOnlyModule.cpp
@@ -4,14 +4,21 @@
 
 #include "Customization/NamingTokensCustomization.h"
 #include "NamingTokenData.h"
+#include "NamingTokens.h"
 
 #include "PropertyEditorModule.h"
 
 void FNamingTokensUncookedOnlyModule::StartupModule()
 {
 	FPropertyEditorModule& PropertyModule = FModuleManager::GetModuleChecked<FPropertyEditorModule>("PropertyEditor");
+
+	PropertyModule.RegisterCustomClassLayout(
+		UNamingTokens::StaticClass()->GetFName(),
+		FOnGetDetailCustomizationInstance::CreateStatic(&FNamingTokensCustomization::MakeInstance)
+	);
+
 	PropertyModule.RegisterCustomPropertyTypeLayout(FNamingTokenData::StaticStruct()->GetFName(),
-		FOnGetPropertyTypeCustomizationInstance::CreateStatic(&FNamingTokensCustomization::MakeInstance));
+		FOnGetPropertyTypeCustomizationInstance::CreateStatic(&FNamingTokensDataCustomization::MakeInstance));
 }
 
 void FNamingTokensUncookedOnlyModule::ShutdownModule()
@@ -20,9 +27,10 @@ void FNamingTokensUncookedOnlyModule::ShutdownModule()
 	{
 		if (FPropertyEditorModule* PropertyEditorModule = FModuleManager::GetModulePtr<FPropertyEditorModule>("PropertyEditor"))
 		{
+			PropertyEditorModule->UnregisterCustomClassLayout(UNamingTokens::StaticClass()->GetFName());
 			PropertyEditorModule->UnregisterCustomPropertyTypeLayout(FNamingTokenData::StaticStruct()->GetFName());
 		}
 	}
 }
 
-IMPLEMENT_MODULE(FNamingTokensUncookedOnlyModule, NamingTokensUncookedOnly)
\ No newline at end of file
+IMPLEMENT_MODULE(FNamingTokensUncookedOnlyModule, NamingTokensUncookedOnly)
diff --git a/Engine/Plugins/Developer/NamingTokens/Source/NamingTokensUncookedOnly/Private/NamingTokensUncookedOnlyModule.h b/Engine/Plugins/Developer/NamingTokens/Source/NamingTokensUncookedOnly/Private/NamingTokensUncookedOnlyModule.h
new file mode 100644
--- /dev/null
+++ b/Engine/Plugins/Developer/NamingTokens/Source/NamingTokensUncookedOnly/Private/NamingTokensUncookedOnlyModule.h
@@ -0,0 +1,13 @@
+﻿// Copyright Epic Games, Inc. All Rights Reserved.
+
+#pragma once
+
+#include "CoreMinimal.h"
+#include "Modules/ModuleManager.h"
+
+class FNamingTokensUncookedOnlyModule : public IModuleInterface
+{
+public:
+    virtual void StartupModule() override;
+    virtual void ShutdownModule() override;
+};
diff --git a/Engine/Plugins/Developer/NamingTokens/Source/NamingTokensUncookedOnly/Public/NamingTokensUncookedOnlyModule.h b/Engine/Plugins/Developer/NamingTokens/Source/NamingTokensUncookedOnly/Public/NamingTokensUncookedOnlyModule.h
deleted file mode 100644
