NamingTokens: Add default global {user} token
#rnx

diff --git a/Engine/Plugins/Developer/NamingTokens/Source/NamingTokens/Private/GlobalNamingTokens.cpp b/Engine/Plugins/Developer/NamingTokens/Source/NamingTokens/Private/GlobalNamingTokens.cpp
--- a/Engine/Plugins/Developer/NamingTokens/Source/NamingTokens/Private/GlobalNamingTokens.cpp
+++ b/Engine/Plugins/Developer/NamingTokens/Source/NamingTokens/Private/GlobalNamingTokens.cpp
@@ -23,6 +23,14 @@ void UGlobalNamingTokens::OnCreateDefaultTokens(TArray<FNamingTokenData>& Tokens
 		})
 	});
 
+	Tokens.Add({
+		TEXT("user"),
+		LOCTEXT("UserName", "User Name"),
+		FNamingTokenData::FTokenProcessorDelegateNative::CreateLambda([]() {
+			return FText::FromString(FApp::GetSessionOwner());
+		})
+	});
+
 	Tokens.Add({
 	    TEXT("yyyy"),
 	    LOCTEXT("Year4Digit", "Year (4 digit)"),
