NamingTokens: Allow cache to be disabled, better supporting tools that need to use tokens off of the game thread.

diff --git a/Engine/Plugins/Developer/NamingTokens/Source/NamingTokens/Private/NamingTokensEngineSubsystem.cpp b/Engine/Plugins/Developer/NamingTokens/Source/NamingTokens/Private/NamingTokensEngineSubsystem.cpp
--- a/Engine/Plugins/Developer/NamingTokens/Source/NamingTokens/Private/NamingTokensEngineSubsystem.cpp
+++ b/Engine/Plugins/Developer/NamingTokens/Source/NamingTokens/Private/NamingTokensEngineSubsystem.cpp
@@ -49,6 +49,11 @@ TArray<FString> UNamingTokensEngineSubsystem::GetAllNamespaces() const
 	return DiscoveredNamespaces.Array();
 }
 
+UNamingTokensEngineSubsystem::UNamingTokensEngineSubsystem()
+	: bIsCacheEnabled(true)
+{
+}
+
 UNamingTokens* UNamingTokensEngineSubsystem::GetNamingTokens(const FString& InNamespace) const
 {
 	// Check cache.
@@ -378,6 +383,20 @@ void UNamingTokensEngineSubsystem::ClearCachedNamingTokens()
 	CachedNamingTokens.Empty();
 }
 
+bool UNamingTokensEngineSubsystem::IsCacheEnabled() const
+{
+	return bIsCacheEnabled;
+}
+
+void UNamingTokensEngineSubsystem::SetCacheEnabled(bool bEnabled)
+{
+	bIsCacheEnabled = bEnabled;
+	if (!bIsCacheEnabled)
+	{
+		ClearCachedNamingTokens();
+	}
+}
+
 UNamingTokens* UNamingTokensEngineSubsystem::LoadNamingToken(const TSoftClassPtr<UNamingTokens>& InTokensClass, const FString& InNamespace) const
 {
 	if (const UClass* Class = InTokensClass.LoadSynchronous())
@@ -388,7 +407,10 @@ UNamingTokens* UNamingTokensEngineSubsystem::LoadNamingToken(const TSoftClassPtr
 		if (ensure(NamingTokens->GetNamespace() == InNamespace))
 		{
 			FScopeLock Lock(&CachedNamingTokensMutex);
-			CachedNamingTokens.Add(InNamespace, NamingTokens);
+			if (IsCacheEnabled())
+			{
+				CachedNamingTokens.Add(InNamespace, NamingTokens);
+			}
 			return NamingTokens;
 		}
 	}
@@ -399,16 +421,19 @@ UNamingTokens* UNamingTokensEngineSubsystem::LoadNamingToken(const TSoftClassPtr
 UNamingTokens* UNamingTokensEngineSubsystem::GetNamingTokenFromCache(const FString& InNamespace, bool bNativeOnly) const
 {
 	FScopeLock Lock(&CachedNamingTokensMutex);
-	if (const TObjectPtr<UNamingTokens>* Tokens = CachedNamingTokens.Find(InNamespace))
+	if (IsCacheEnabled())
 	{
-		if (!bNativeOnly)
+		if (const TObjectPtr<UNamingTokens>* Tokens = CachedNamingTokens.Find(InNamespace))
 		{
-			return *Tokens;
-		}
+			if (!bNativeOnly)
+			{
+				return *Tokens;
+			}
 
-		if (bNativeOnly && (*Tokens)->GetClass()->IsNative())
-		{
-			return *Tokens;
+			if (bNativeOnly && (*Tokens)->GetClass()->IsNative())
+			{
+				return *Tokens;
+			}
 		}
 	}
 
diff --git a/Engine/Plugins/Developer/NamingTokens/Source/NamingTokens/Public/NamingTokensEngineSubsystem.h b/Engine/Plugins/Developer/NamingTokens/Source/NamingTokens/Public/NamingTokensEngineSubsystem.h
--- a/Engine/Plugins/Developer/NamingTokens/Source/NamingTokens/Public/NamingTokensEngineSubsystem.h
+++ b/Engine/Plugins/Developer/NamingTokens/Source/NamingTokens/Public/NamingTokensEngineSubsystem.h
@@ -48,6 +48,8 @@ class UNamingTokensEngineSubsystem : public UEngineSubsystem
 	GENERATED_BODY()
 
 public:
+	UE_API UNamingTokensEngineSubsystem();
+	
 	/**
 	 * Lookup naming tokens given a namespace. This will look first in cached naming tokens,
 	 * then native classes, then blueprint classes. Call ClearCachedNamingTokens to reset the cache
@@ -149,6 +151,12 @@ public:
 	 */
 	UFUNCTION(BlueprintCallable, Category = "NamingTokens")
 	UE_API void ClearCachedNamingTokens();
+
+	/** If the cache is currently enabled. */
+	UE_API bool IsCacheEnabled() const;
+
+	/** Configure whether the cache is enabled or not. */
+	UE_API void SetCacheEnabled(bool bEnabled);
 	
 private:
 	/** Naming tokens currently loaded from assets. */
@@ -164,6 +172,9 @@ private:
 
 	/** Filters that will execute just before evaluating a token string */
 	TMap<FName, FFilterNamespace> FilterNamespaceDelegates;
+
+	/** If the cache is currently enabled. */
+	bool bIsCacheEnabled;
 };
 
 #undef UE_API
