NamingTokens: Fix crash if trying to reference a deleted namespace from the cache.

diff --git a/Engine/Plugins/Developer/NamingTokens/Source/NamingTokens/Private/NamingTokensEngineSubsystem.cpp b/Engine/Plugins/Developer/NamingTokens/Source/NamingTokens/Private/NamingTokensEngineSubsystem.cpp
--- a/Engine/Plugins/Developer/NamingTokens/Source/NamingTokens/Private/NamingTokensEngineSubsystem.cpp
+++ b/Engine/Plugins/Developer/NamingTokens/Source/NamingTokens/Private/NamingTokensEngineSubsystem.cpp
@@ -425,6 +425,13 @@ UNamingTokens* UNamingTokensEngineSubsystem::GetNamingTokenFromCache(const FStri
 	{
 		if (const TObjectPtr<UNamingTokens>* Tokens = CachedNamingTokens.Find(InNamespace))
 		{
+			if ((*Tokens) == nullptr)
+			{
+				// NamingTokens may have been deleted, with a new asset using the same namespace then added afterward.
+				CachedNamingTokens.Remove(InNamespace);
+				return nullptr;
+			}
+			
 			if (!bNativeOnly)
 			{
 				return *Tokens;
