NamingTokens: Add EvaluateTokenList API to evaluate specific tokens and return their evaluation data.
- Add TokenNamespace field to FNamingTokenValueData so the namespace is always returned. Remove ProvidedNamespace as it was not in use anywhere. No deprecation as this is an experimental plugin.
- Add CombineNamespaceAndTokenKey utility function in case the key and namespace need to be recombined.

diff --git a/Engine/Plugins/Developer/NamingTokens/Source/NamingTokens/Private/NamingTokens.cpp b/Engine/Plugins/Developer/NamingTokens/Source/NamingTokens/Private/NamingTokens.cpp
--- a/Engine/Plugins/Developer/NamingTokens/Source/NamingTokens/Private/NamingTokens.cpp
+++ b/Engine/Plugins/Developer/NamingTokens/Source/NamingTokens/Private/NamingTokens.cpp
@@ -190,8 +190,8 @@ FNamingTokenResultData UNamingTokens::EvaluateTokenText(const FText& InTokenText
 			}
 			
 			ValueData.TokenKey = SelectedToken ? SelectedToken->TokenKey : TokenKey;
+			ValueData.TokenNamespace = Namespace;
 			ValueData.TokenValue = TokenValue;
-			ValueData.ProvidedNamespace = TokenNamespace;
 			Result.TokenValues.Add(MoveTemp(ValueData));
 		}
 	}
@@ -317,10 +317,11 @@ FText UNamingTokens::ProcessToken(const FNamingTokenData& InToken, const FString
 		UE_LOG(LogNamingTokens, Error, TEXT("Could not find function process for token %s"), *FormattedToken);
 		return ProcessedToken;
 	}
-		
+
+	// Replace both {token} and {namespace:token} in the formatted text.
 	FFormatNamedArguments Args;
 	Args.Add(InUserTokenString, ProcessedToken);
-	Args.Add(FString::Printf(TEXT("%s%s%s"), *GetNamespace(), *UE::NamingTokens::Utils::GetNamespaceDelimiter(), *InUserTokenString), ProcessedToken);
+	Args.Add(UE::NamingTokens::Utils::CombineNamespaceAndTokenKey(*GetNamespace(), *InUserTokenString), ProcessedToken);
 
 	InOutFormattedText = FText::Format(InOutFormattedText, Args);
 	return ProcessedToken;
diff --git a/Engine/Plugins/Developer/NamingTokens/Source/NamingTokens/Private/NamingTokensEngineSubsystem.cpp b/Engine/Plugins/Developer/NamingTokens/Source/NamingTokens/Private/NamingTokensEngineSubsystem.cpp
--- a/Engine/Plugins/Developer/NamingTokens/Source/NamingTokens/Private/NamingTokensEngineSubsystem.cpp
+++ b/Engine/Plugins/Developer/NamingTokens/Source/NamingTokens/Private/NamingTokensEngineSubsystem.cpp
@@ -287,6 +287,23 @@ FNamingTokenResultData UNamingTokensEngineSubsystem::EvaluateTokenString(const F
 	return EvaluateTokenString(InTokenString, InFilter, {});
 }
 
+TArray<FNamingTokenValueData> UNamingTokensEngineSubsystem::EvaluateTokenList(const TArray<FString>& InTokenList, const FNamingTokenFilterArgs& InFilter,
+                                                                           const TArray<UObject*>& InContexts)
+{
+	const FString TokenStringCombined = FString::JoinBy(InTokenList, TEXT(", "),
+		[](const FString& InToken)
+		{
+			return UE::NamingTokens::Utils::CreateFormattedToken(FNamingTokenData(InToken));
+		});
+	FNamingTokenResultData Result = EvaluateTokenString(TokenStringCombined, InFilter, InContexts);
+	return Result.TokenValues;
+}
+
+TArray<FNamingTokenValueData> UNamingTokensEngineSubsystem::EvaluateTokenList(const TArray<FString>& InTokenList, const FNamingTokenFilterArgs& InFilter)
+{
+	return EvaluateTokenList(InTokenList, InFilter, {});
+}
+
 void UNamingTokensEngineSubsystem::RegisterGlobalNamespace(const FString& InNamespace)
 {
 	if (IsGlobalNamespaceRegistered(InNamespace))
diff --git a/Engine/Plugins/Developer/NamingTokens/Source/NamingTokens/Private/Utils/NamingTokenUtils.cpp b/Engine/Plugins/Developer/NamingTokens/Source/NamingTokens/Private/Utils/NamingTokenUtils.cpp
--- a/Engine/Plugins/Developer/NamingTokens/Source/NamingTokens/Private/Utils/NamingTokenUtils.cpp
+++ b/Engine/Plugins/Developer/NamingTokens/Source/NamingTokens/Private/Utils/NamingTokenUtils.cpp
@@ -98,6 +98,16 @@ FString UE::NamingTokens::Utils::RemoveNamespaceFromTokenKey(const FString& InTo
 	return Result;
 }
 
+FString UE::NamingTokens::Utils::CombineNamespaceAndTokenKey(const FString& InNamespace, const FString& InTokenKey)
+{
+	if (InNamespace.IsEmpty())
+	{
+		return InTokenKey;
+	}
+	
+	return FString::Printf(TEXT("%s%s%s"), *InNamespace, *GetNamespaceDelimiter(), *InTokenKey);
+}
+
 bool UE::NamingTokens::Utils::ValidateName(const FString& InName, FText& OutErrorMessage)
 {
 	OutErrorMessage = FText::GetEmpty();
diff --git a/Engine/Plugins/Developer/NamingTokens/Source/NamingTokens/Public/NamingTokenData.h b/Engine/Plugins/Developer/NamingTokens/Source/NamingTokens/Public/NamingTokenData.h
--- a/Engine/Plugins/Developer/NamingTokens/Source/NamingTokens/Public/NamingTokenData.h
+++ b/Engine/Plugins/Developer/NamingTokens/Source/NamingTokens/Public/NamingTokenData.h
@@ -19,9 +19,9 @@ struct FNamingTokenValueData
 	/** The token key. */
 	UPROPERTY(BlueprintReadOnly, Category = "NamingTokens")
 	FString TokenKey;
-	/** The namespace of the token, if it was originally provided during evaluation. */
+	/** The namespace of the token key. */
 	UPROPERTY(BlueprintReadOnly, Category = "NamingTokens")
-	FString ProvidedNamespace;
+	FString TokenNamespace;
 	/** The evaluated token text. */
 	UPROPERTY(BlueprintReadOnly, Category = "NamingTokens")
 	FText TokenValue;
diff --git a/Engine/Plugins/Developer/NamingTokens/Source/NamingTokens/Public/NamingTokensEngineSubsystem.h b/Engine/Plugins/Developer/NamingTokens/Source/NamingTokens/Public/NamingTokensEngineSubsystem.h
--- a/Engine/Plugins/Developer/NamingTokens/Source/NamingTokens/Public/NamingTokensEngineSubsystem.h
+++ b/Engine/Plugins/Developer/NamingTokens/Source/NamingTokens/Public/NamingTokensEngineSubsystem.h
@@ -106,6 +106,19 @@ public:
 	UFUNCTION(BlueprintCallable, Category = "NamingTokens", meta = (AutoCreateRefTerm = "InFilter,InContexts"))
 	UE_API FNamingTokenResultData EvaluateTokenString(const FString& InTokenString, const FNamingTokenFilterArgs& InFilter, const TArray<UObject*>& InContexts);
 	UE_API FNamingTokenResultData EvaluateTokenString(const FString& InTokenString, const FNamingTokenFilterArgs& InFilter = FNamingTokenFilterArgs());
+
+	/**
+	 * Given a list of tokens, return a list of all found tokens and their values.
+	 * 
+	 * @param InTokenList A list of raw tokens. Tokens can include their namespace, but do not include brackets.
+	 * @param InFilter [Optional] Filter to determine which namespaces to use.
+	 * @param InContexts [Optional] Context objects to pass to naming tokens.
+	 *
+	 * @return Evaluation data for each token.
+	 */
+	UFUNCTION(BlueprintCallable, Category = "NamingTokens", meta = (AutoCreateRefTerm = "InFilter,InContexts"))
+	UE_API TArray<FNamingTokenValueData> EvaluateTokenList(const TArray<FString>& InTokenList, const FNamingTokenFilterArgs& InFilter, const TArray<UObject*>& InContexts);
+	UE_API TArray<FNamingTokenValueData> EvaluateTokenList(const TArray<FString>& InTokenList, const FNamingTokenFilterArgs& InFilter = FNamingTokenFilterArgs());
 	
 	/**
 	 * Register tokens as a global namespace. This prevents the need to include the namespace in a token string.
diff --git a/Engine/Plugins/Developer/NamingTokens/Source/NamingTokens/Public/Utils/NamingTokenUtils.h b/Engine/Plugins/Developer/NamingTokens/Source/NamingTokens/Public/Utils/NamingTokenUtils.h
--- a/Engine/Plugins/Developer/NamingTokens/Source/NamingTokens/Public/Utils/NamingTokenUtils.h
+++ b/Engine/Plugins/Developer/NamingTokens/Source/NamingTokens/Public/Utils/NamingTokenUtils.h
@@ -40,6 +40,14 @@ namespace UE::NamingTokens::Utils
 	 */
 	NAMINGTOKENS_API FString RemoveNamespaceFromTokenKey(const FString& InTokenKey);
 
+	/**
+	 * Combine the namespace and token key using the delimiter.
+	 * @param InNamespace The namespace. If empty no combination is performed.
+	 * @param InTokenKey The token key.
+	 * @return The combined key.
+	 */
+	NAMINGTOKENS_API FString CombineNamespaceAndTokenKey(const FString& InNamespace, const FString& InTokenKey);
+
 	/** Validates a name used for a token key or namespace. */
 	NAMINGTOKENS_API bool ValidateName(const FString& InName, FText& OutErrorMessage);
 }
\ No newline at end of file
