NamingTokens: Evaluation data no longer initializes datetime in the constructor, fixing an issue with CDO stability.

diff --git a/Engine/Plugins/Developer/NamingTokens/Source/NamingTokens/Private/NamingTokens.cpp b/Engine/Plugins/Developer/NamingTokens/Source/NamingTokens/Private/NamingTokens.cpp
--- a/Engine/Plugins/Developer/NamingTokens/Source/NamingTokens/Private/NamingTokens.cpp
+++ b/Engine/Plugins/Developer/NamingTokens/Source/NamingTokens/Private/NamingTokens.cpp
@@ -78,6 +78,7 @@ void UNamingTokens::OnCreateDefaultTokens(TArray<FNamingTokenData>& Tokens)
 FNamingTokenResultData UNamingTokens::EvaluateTokenText(const FText& InTokenText, const TArray<UObject*>& InContexts)
 {
 	FNamingTokensEvaluationData EvaluationData;
+	EvaluationData.Initialize();
 	EvaluationData.Contexts = InContexts;
 	return EvaluateTokenText(InTokenText, EvaluationData);
 }
diff --git a/Engine/Plugins/Developer/NamingTokens/Source/NamingTokens/Private/NamingTokensEngineSubsystem.cpp b/Engine/Plugins/Developer/NamingTokens/Source/NamingTokens/Private/NamingTokensEngineSubsystem.cpp
--- a/Engine/Plugins/Developer/NamingTokens/Source/NamingTokens/Private/NamingTokensEngineSubsystem.cpp
+++ b/Engine/Plugins/Developer/NamingTokens/Source/NamingTokens/Private/NamingTokensEngineSubsystem.cpp
@@ -217,6 +217,7 @@ FNamingTokenResultData UNamingTokensEngineSubsystem::EvaluateTokenText(const FTe
 	Namespaces.Append(InFilter.AdditionalNamespacesToInclude); // Tokens could have been written without the namespace but should still be scoped to the filter.
 
 	FNamingTokensEvaluationData EvaluationData;
+	EvaluationData.Initialize();
 	EvaluationData.Contexts = InContexts;
 	EvaluationData.bForceCaseSensitive = InFilter.bForceCaseSensitive;
 	for (const FString& Namespace : Namespaces)
diff --git a/Engine/Plugins/Developer/NamingTokens/Source/NamingTokens/Private/NamingTokensEvaluationData.cpp b/Engine/Plugins/Developer/NamingTokens/Source/NamingTokens/Private/NamingTokensEvaluationData.cpp
--- a/Engine/Plugins/Developer/NamingTokens/Source/NamingTokens/Private/NamingTokensEvaluationData.cpp
+++ b/Engine/Plugins/Developer/NamingTokens/Source/NamingTokens/Private/NamingTokensEvaluationData.cpp
@@ -3,6 +3,10 @@
 #include "NamingTokensEvaluationData.h"
 
 FNamingTokensEvaluationData::FNamingTokensEvaluationData()
+{
+}
+
+void FNamingTokensEvaluationData::Initialize()
 {
 	CurrentDateTime = FDateTime::Now();
 }
diff --git a/Engine/Plugins/Developer/NamingTokens/Source/NamingTokens/Public/NamingTokensEvaluationData.h b/Engine/Plugins/Developer/NamingTokens/Source/NamingTokens/Public/NamingTokensEvaluationData.h
--- a/Engine/Plugins/Developer/NamingTokens/Source/NamingTokens/Public/NamingTokensEvaluationData.h
+++ b/Engine/Plugins/Developer/NamingTokens/Source/NamingTokens/Public/NamingTokensEvaluationData.h
@@ -17,6 +17,9 @@ struct FNamingTokensEvaluationData
 
 public:
 	FNamingTokensEvaluationData();
+
+	/** Initialize any constant data that isn't safe to do so in the ctor, like DateTime. */
+	void Initialize();
 	
 	/** The current date time. */
 	UPROPERTY(BlueprintReadOnly, Category = "NamingTokens", meta = (IgnoreForMemberInitializationTest))
