NamingTokens: Add UMG widget for displaying token text, both as raw tokens and resolved tokens.

diff --git a/Engine/Plugins/Developer/NamingTokens/NamingTokens.uplugin b/Engine/Plugins/Developer/NamingTokens/NamingTokens.uplugin
--- a/Engine/Plugins/Developer/NamingTokens/NamingTokens.uplugin
+++ b/Engine/Plugins/Developer/NamingTokens/NamingTokens.uplugin
@@ -23,6 +23,11 @@
 			"Name": "NamingTokensUncookedOnly",
 			"Type": "UncookedOnly",
 			"LoadingPhase": "Default"
+		},
+		{
+			"Name": "NamingTokensUI",
+			"Type": "Runtime",
+			"LoadingPhase": "Default"
 		}
 	],
 	"Plugins": [
diff --git a/Engine/Plugins/Developer/NamingTokens/Source/NamingTokensUI/NamingTokensUI.Build.cs b/Engine/Plugins/Developer/NamingTokens/Source/NamingTokensUI/NamingTokensUI.Build.cs
new file mode 100644
--- /dev/null
+++ b/Engine/Plugins/Developer/NamingTokens/Source/NamingTokensUI/NamingTokensUI.Build.cs
@@ -0,0 +1,40 @@
+﻿// Copyright Epic Games, Inc. All Rights Reserved.
+
+using UnrealBuildTool;
+
+public class NamingTokensUI : ModuleRules
+{
+    public NamingTokensUI(ReadOnlyTargetRules Target) : base(Target)
+    {
+        PCHUsage = ModuleRules.PCHUsageMode.UseExplicitOrSharedPCHs;
+
+        PublicDependencyModuleNames.AddRange(
+            new string[]
+            {
+                "Core",
+                "CoreUObject",
+                "NamingTokens",
+                "UMG"
+            }
+        );
+
+        PrivateDependencyModuleNames.AddRange(
+            new string[]
+            {
+                "Engine",
+                "InputCore",
+                "Slate",
+                "SlateCore"
+            }
+        );
+        
+        if (Target.Type == TargetType.Editor)
+        {
+	        PublicDependencyModuleNames.AddRange(
+		        new string[] {
+			        "UnrealEd",
+		        }
+	        );
+        }
+    }
+}
\ No newline at end of file
diff --git a/Engine/Plugins/Developer/NamingTokens/Source/NamingTokensUI/Private/NamingTokensEditableText.cpp b/Engine/Plugins/Developer/NamingTokens/Source/NamingTokensUI/Private/NamingTokensEditableText.cpp
new file mode 100644
--- /dev/null
+++ b/Engine/Plugins/Developer/NamingTokens/Source/NamingTokensUI/Private/NamingTokensEditableText.cpp
@@ -0,0 +1,117 @@
+// Copyright Epic Games, Inc. All Rights Reserved.
+
+#include "NamingTokensEditableText.h"
+
+#include "SNamingTokensEditableTextBox.h"
+
+UNamingTokensEditableText::UNamingTokensEditableText()
+{
+	ArgumentStyle = FAppStyle::GetWidgetStyle<FTextBlockStyle>("RichTextBlock.BoldHighlight");
+}
+
+const FText& UNamingTokensEditableText::GetResolvedText() const
+{
+	check(MyNamingTokensEditableText.IsValid());
+	return MyNamingTokensEditableText->GetResolvedText();
+}
+
+const FText& UNamingTokensEditableText::GetTokenizedText() const
+{
+	check(MyNamingTokensEditableText.IsValid());
+	return MyNamingTokensEditableText->GetTokenizedText();
+}
+
+void UNamingTokensEditableText::SetWidgetArgumentStyle(const FTextBlockStyle& InWidgetStyle)
+{
+	check(MyNamingTokensEditableText.IsValid());
+	MyNamingTokensEditableText->SetArgumentStyle(&InWidgetStyle);
+}
+
+#if WITH_EDITOR
+void UNamingTokensEditableText::PostEditChangeProperty(FPropertyChangedEvent& PropertyChangedEvent)
+{
+	Super::PostEditChangeProperty(PropertyChangedEvent);
+
+	PRAGMA_DISABLE_DEPRECATION_WARNINGS
+	if (PropertyChangedEvent.GetPropertyName() == GET_MEMBER_NAME_CHECKED(UNamingTokensEditableText, Text)
+		|| PropertyChangedEvent.GetPropertyName() == GET_MEMBER_NAME_CHECKED(UNamingTokensEditableText, bDisplayResolvedTextInDesigner))
+	{
+		if (MyNamingTokensEditableText.IsValid())
+		{
+			MyNamingTokensEditableText->SetTokenizedText(Text);
+		}
+		SynchronizeProperties();
+	}
+	PRAGMA_ENABLE_DEPRECATION_WARNINGS
+}
+#endif
+
+TSharedRef<SWidget> UNamingTokensEditableText::RebuildWidget()
+{
+	// We only use text box style, but slate needs editable text box style.
+	EditableTextBoxStyle = MakeShared<FEditableTextBoxStyle>(FAppStyle::Get().GetWidgetStyle<FEditableTextBoxStyle>("NormalEditableTextBox"));
+	EditableTextBoxStyle->SetTextStyle(WidgetStyle);
+	
+	MyNamingTokensEditableText = SNew(SNamingTokensEditableTextBox)
+		.Text(GetText())
+		.IsReadOnly(GetIsReadOnly())
+#if WITH_EDITOR
+		.CanDisplayResolvedText(MakeAttributeUObject(this, &UNamingTokensEditableText::CanDisplayResolvedTextInEditor))
+#endif
+		.AllowMultiLine(bIsMultiline)
+		.Style(&*EditableTextBoxStyle)
+		.ArgumentStyle(&ArgumentStyle);
+
+	// Make sure our primary editable text (managed by parent) points to the editable text we create from the naming tokens slate text box.
+	MyMultiLineEditableText	= MyNamingTokensEditableText->GetEditableText();
+
+	return MyMultiLineEditableText.ToSharedRef();
+}
+
+void UNamingTokensEditableText::SynchronizeProperties()
+{
+	// Skip parent call, we want to avoid setting the text directly and use our bound text attribute instead.
+	UWidget::SynchronizeProperties();
+
+	if (!MyMultiLineEditableText.IsValid())
+	{
+		return;
+	}
+	PRAGMA_DISABLE_DEPRECATION_WARNINGS
+	const TAttribute<FText> HintTextBinding = PROPERTY_BINDING(FText, HintText);
+	PRAGMA_ENABLE_DEPRECATION_WARNINGS
+
+	EditableTextBoxStyle->SetTextStyle(WidgetStyle);
+	MyNamingTokensEditableText->SetNormalStyle(&*EditableTextBoxStyle);
+	MyNamingTokensEditableText->SetArgumentStyle(&ArgumentStyle);
+	
+	MyMultiLineEditableText->SetText(MyNamingTokensEditableText->GetBoundText());
+	MyMultiLineEditableText->SetHintText(HintTextBinding);
+	MyMultiLineEditableText->SetAllowContextMenu(AllowContextMenu);
+	MyMultiLineEditableText->SetIsReadOnly(GetIsReadOnly());
+	MyMultiLineEditableText->SetVirtualKeyboardDismissAction(VirtualKeyboardDismissAction);
+	MyMultiLineEditableText->SetSelectAllTextWhenFocused(GetSelectAllTextWhenFocused());
+	MyMultiLineEditableText->SetClearTextSelectionOnFocusLoss(GetClearTextSelectionOnFocusLoss());
+	MyMultiLineEditableText->SetRevertTextOnEscape(GetRevertTextOnEscape());
+	MyMultiLineEditableText->SetClearKeyboardFocusOnCommit(GetClearKeyboardFocusOnCommit());
+
+	Super::SynchronizeTextLayoutProperties(*MyMultiLineEditableText);
+}
+
+void UNamingTokensEditableText::ReleaseSlateResources(bool bReleaseChildren)
+{
+	Super::ReleaseSlateResources(bReleaseChildren);
+	MyNamingTokensEditableText.Reset();
+}
+
+#if WITH_EDITOR
+bool UNamingTokensEditableText::CanDisplayResolvedTextInEditor() const
+{
+	if (IsDesignTime())
+	{
+		return bDisplayResolvedTextInDesigner;
+	}
+
+	return true;
+}
+#endif
\ No newline at end of file
diff --git a/Engine/Plugins/Developer/NamingTokens/Source/NamingTokensUI/Private/NamingTokensStringSyntaxHighlighterMarshaller.cpp b/Engine/Plugins/Developer/NamingTokens/Source/NamingTokensUI/Private/NamingTokensStringSyntaxHighlighterMarshaller.cpp
new file mode 100644
--- /dev/null
+++ b/Engine/Plugins/Developer/NamingTokens/Source/NamingTokensUI/Private/NamingTokensStringSyntaxHighlighterMarshaller.cpp
@@ -0,0 +1,104 @@
+﻿// Copyright Epic Games, Inc. All Rights Reserved.
+
+#include "NamingTokensStringSyntaxHighlighterMarshaller.h"
+
+#include "Containers/StringFwd.h"
+#include "Containers/StringView.h"
+#include "Framework/Text/IRun.h"
+#include "Framework/Text/ISlateRun.h"
+#include "Framework/Text/SlateTextRun.h"
+#include "Framework/Text/TextLayout.h"
+#include "HAL/Platform.h"
+#include "HAL/PlatformCrt.h"
+#include "Templates/UnrealTemplate.h"
+
+constexpr static const TCHAR* Operators[] = {
+	TEXT("{"),
+	TEXT("}")
+};
+
+TSharedRef<FNamingTokensStringSyntaxHighlighterMarshaller> FNamingTokensStringSyntaxHighlighterMarshaller::Create(const FSyntaxTextStyle& InSyntaxTextStyle)
+{
+	TArray<FSyntaxTokenizer::FRule> TokenizerRules;
+
+	// operators
+	for(const TCHAR* Operator : Operators)
+	{
+		TokenizerRules.Emplace(FSyntaxTokenizer::FRule(Operator));
+	}	
+
+	return MakeShared<FNamingTokensStringSyntaxHighlighterMarshaller>(FSyntaxTokenizer::Create(TokenizerRules), InSyntaxTextStyle);
+}
+
+void FNamingTokensStringSyntaxHighlighterMarshaller::ParseTokens(const FString& SourceString, FTextLayout& TargetTextLayout, TArray<ISyntaxTokenizer::FTokenizedLine> TokenizedLines)
+{
+	enum class EParseState : uint8
+	{
+		None,
+		LookingForArgument,
+	};
+
+	TArray<FTextLayout::FNewLineData> LinesToAdd;
+	LinesToAdd.Reserve(TokenizedLines.Num());
+
+	// Parse the tokens, generating the styled runs for each line
+	EParseState ParseState = EParseState::None;
+	for(const ISyntaxTokenizer::FTokenizedLine& TokenizedLine : TokenizedLines)
+	{
+		TSharedRef<FString> ModelString = MakeShared<FString>();
+		TArray<TSharedRef<IRun>> Runs;
+
+		for(const ISyntaxTokenizer::FToken& Token : TokenizedLine.Tokens)
+		{
+			const FStringView TokenText(*SourceString + Token.Range.BeginIndex, Token.Range.Len());
+
+			const FTextRange ModelRange(ModelString->Len(), ModelString->Len() + TokenText.Len());
+			ModelString->Append(TokenText);
+
+			FRunInfo RunInfo(TEXT("SyntaxHighlight.Template.Normal"));
+			FTextBlockStyle TextBlockStyle = SyntaxTextStyle.NormalTextStyle;
+
+			bool bHasMatchedSyntax = false;
+			if (Token.Type == ISyntaxTokenizer::ETokenType::Syntax)
+			{
+				if(ParseState == EParseState::None && TokenText == TEXT("{"))
+				{
+					RunInfo.Name = TEXT("SyntaxHighlight.Template.Argument");
+					TextBlockStyle = SyntaxTextStyle.ArgumentTextStyle;
+					ParseState = EParseState::LookingForArgument;
+					bHasMatchedSyntax = true;
+				}
+				else if(ParseState == EParseState::LookingForArgument && TokenText == TEXT("}"))
+				{
+					RunInfo.Name = TEXT("SyntaxHighlight.Template.Normal");
+					TextBlockStyle = SyntaxTextStyle.ArgumentTextStyle; // Closing brace should still use this style
+					ParseState = EParseState::None;
+				}
+			}
+			
+			// It's possible that we fail to match a syntax token if we're in a state where it isn't parsed
+			// In this case, we treat it as a literal token
+			if (Token.Type == ISyntaxTokenizer::ETokenType::Literal || !bHasMatchedSyntax)
+			{
+				if(ParseState == EParseState::LookingForArgument)
+				{
+					RunInfo.Name = TEXT("SyntaxHighlight.Template.Argument");
+					TextBlockStyle = SyntaxTextStyle.ArgumentTextStyle;
+				}
+			}
+
+			TSharedRef< ISlateRun > Run = FSlateTextRun::Create(RunInfo, ModelString, TextBlockStyle, ModelRange);
+			Runs.Add(Run);
+		}
+
+		LinesToAdd.Emplace(MoveTemp(ModelString), MoveTemp(Runs));
+	}
+
+	TargetTextLayout.AddLines(LinesToAdd);
+}
+
+FNamingTokensStringSyntaxHighlighterMarshaller::FNamingTokensStringSyntaxHighlighterMarshaller(TSharedPtr< ISyntaxTokenizer > InTokenizer, const FSyntaxTextStyle& InSyntaxTextStyle)
+	: FSyntaxHighlighterTextLayoutMarshaller(MoveTemp(InTokenizer))
+	, SyntaxTextStyle(InSyntaxTextStyle)
+{
+}
diff --git a/Engine/Plugins/Developer/NamingTokens/Source/NamingTokensUI/Private/NamingTokensStringSyntaxHighlighterMarshaller.h b/Engine/Plugins/Developer/NamingTokens/Source/NamingTokensUI/Private/NamingTokensStringSyntaxHighlighterMarshaller.h
new file mode 100644
--- /dev/null
+++ b/Engine/Plugins/Developer/NamingTokens/Source/NamingTokensUI/Private/NamingTokensStringSyntaxHighlighterMarshaller.h
@@ -0,0 +1,65 @@
+﻿// Copyright Epic Games, Inc. All Rights Reserved.
+
+#pragma once
+
+#include "Containers/Array.h"
+#include "Containers/UnrealString.h"
+#include "Framework/Text/SyntaxHighlighterTextLayoutMarshaller.h"
+#include "Framework/Text/SyntaxTokenizer.h"
+#include "Styling/ISlateStyle.h"
+#include "Styling/SlateTypes.h"
+#include "Templates/SharedPointer.h"
+
+class FTextLayout;
+
+/**
+ * Get/set the raw text to/from a text layout, and also inject syntax highlighting for our rich-text markup
+ */
+class FNamingTokensStringSyntaxHighlighterMarshaller : public FSyntaxHighlighterTextLayoutMarshaller
+{
+public:
+	struct FSyntaxTextStyle
+	{
+		FSyntaxTextStyle()
+			: NormalTextStyle(FAppStyle::GetWidgetStyle<FTextBlockStyle>(TEXT("NormalText")))
+			, ArgumentTextStyle(FAppStyle::GetWidgetStyle<FTextBlockStyle>("RichTextBlock.BoldHighlight"))
+		{
+		}
+		
+		FSyntaxTextStyle(const FTextBlockStyle& InNormalTextStyle, const FTextBlockStyle& InArgumentTextStyle)
+			: NormalTextStyle(InNormalTextStyle)
+			, ArgumentTextStyle(InArgumentTextStyle)
+		{
+		}
+		
+		FTextBlockStyle NormalTextStyle;
+		FTextBlockStyle ArgumentTextStyle;
+	};
+
+	static TSharedRef<FNamingTokensStringSyntaxHighlighterMarshaller> Create(const FSyntaxTextStyle& InSyntaxTextStyle);
+
+	virtual ~FNamingTokensStringSyntaxHighlighterMarshaller() override = default;
+
+	/** Set the normal style of our syntax. */
+	void SetNormalStyle(const FTextBlockStyle& InStyle)
+	{
+		SyntaxTextStyle.NormalTextStyle = InStyle;
+	}
+
+	/** Set the argument style for our syntax. */
+	void SetArgumentStyle(const FTextBlockStyle& InStyle)
+	{
+		SyntaxTextStyle.ArgumentTextStyle = InStyle;
+	}
+	
+protected:
+	// Allows MakeShared with private constructor
+	friend class SharedPointerInternals::TIntrusiveReferenceController<FNamingTokensStringSyntaxHighlighterMarshaller, ESPMode::ThreadSafe>;
+	
+	virtual void ParseTokens(const FString& SourceString, FTextLayout& TargetTextLayout, TArray<ISyntaxTokenizer::FTokenizedLine> TokenizedLines) override;
+
+	FNamingTokensStringSyntaxHighlighterMarshaller(TSharedPtr<ISyntaxTokenizer> InTokenizer, const FSyntaxTextStyle& InSyntaxTextStyle);
+
+	/** Styles used to display the text */
+	FSyntaxTextStyle SyntaxTextStyle;
+};
diff --git a/Engine/Plugins/Developer/NamingTokens/Source/NamingTokensUI/Private/NamingTokensUIModule.cpp b/Engine/Plugins/Developer/NamingTokens/Source/NamingTokensUI/Private/NamingTokensUIModule.cpp
new file mode 100644
--- /dev/null
+++ b/Engine/Plugins/Developer/NamingTokens/Source/NamingTokensUI/Private/NamingTokensUIModule.cpp
@@ -0,0 +1,13 @@
+﻿// Copyright Epic Games, Inc. All Rights Reserved.
+
+#include "NamingTokensUIModule.h"
+
+void FNamingTokensUIModule::StartupModule()
+{
+}
+
+void FNamingTokensUIModule::ShutdownModule()
+{
+}
+
+IMPLEMENT_MODULE(FNamingTokensUIModule, NamingTokensUI)
diff --git a/Engine/Plugins/Developer/NamingTokens/Source/NamingTokensUI/Private/NamingTokensUIModule.h b/Engine/Plugins/Developer/NamingTokens/Source/NamingTokensUI/Private/NamingTokensUIModule.h
new file mode 100644
--- /dev/null
+++ b/Engine/Plugins/Developer/NamingTokens/Source/NamingTokensUI/Private/NamingTokensUIModule.h
@@ -0,0 +1,12 @@
+﻿// Copyright Epic Games, Inc. All Rights Reserved.
+
+#pragma once
+
+#include "Modules/ModuleManager.h"
+
+class FNamingTokensUIModule : public IModuleInterface
+{
+public:
+    virtual void StartupModule() override;
+    virtual void ShutdownModule() override;
+};
diff --git a/Engine/Plugins/Developer/NamingTokens/Source/NamingTokensUI/Private/SNamingTokensEditableTextBox.cpp b/Engine/Plugins/Developer/NamingTokens/Source/NamingTokensUI/Private/SNamingTokensEditableTextBox.cpp
new file mode 100644
--- /dev/null
+++ b/Engine/Plugins/Developer/NamingTokens/Source/NamingTokensUI/Private/SNamingTokensEditableTextBox.cpp
@@ -0,0 +1,439 @@
+// Copyright Epic Games, Inc. All Rights Reserved.
+
+#include "SNamingTokensEditableTextBox.h"
+
+#include "NamingTokensEngineSubsystem.h"
+#include "NamingTokensStringSyntaxHighlighterMarshaller.h"
+#include "Utils/NamingTokenUtils.h"
+
+#include "Engine/Engine.h"
+#include "Internationalization/Text.h"
+#include "Layout/Margin.h"
+#include "SlateOptMacros.h"
+#include "Styling/ISlateStyle.h"
+#include "Styling/SlateTypes.h"
+#include "Styling/StyleColors.h"
+#include "Widgets/Images/SImage.h"
+#include "Widgets/Layout/SBox.h"
+#include "Widgets/Notifications/SPopUpErrorText.h"
+#include "Widgets/SBoxPanel.h"
+
+#define LOCTEXT_NAMESPACE "SNamingTokensEditableTextBox"
+
+BEGIN_SLATE_FUNCTION_BUILD_OPTIMIZATION
+
+void SNamingTokensEditableTextBox::Construct(const FArguments& InArgs)
+{
+	if (InArgs._Style == nullptr)
+	{
+		SetNormalStyle(&FAppStyle::Get().GetWidgetStyle<FEditableTextBoxStyle>("NormalEditableTextBox"));		
+	}
+	else
+	{
+		SetNormalStyle(InArgs._Style);
+	}
+
+	if (InArgs._ArgumentStyle != nullptr)
+	{
+		SetArgumentStyle(InArgs._ArgumentStyle);
+	}
+	
+	Marshaller = FNamingTokensStringSyntaxHighlighterMarshaller::Create(
+		FNamingTokensStringSyntaxHighlighterMarshaller::FSyntaxTextStyle(TextBoxStyle->TextStyle,\
+			*InArgs._ArgumentStyle));
+	
+	static const FVector2D ScrollBarThickness(9.0f, 9.0f);
+	static const FMargin ErrorWidgetPadding(3, 0); 
+
+	TokenizedText = InArgs._Text;
+	ResolvedText = InArgs._ResolvedText;
+	ValidArguments = InArgs._ValidArguments;
+
+	OnValidateTokenizedText = InArgs._OnValidateTokenizedText;
+	OnTokenizedTextChanged = InArgs._OnTextChanged;
+	OnTokenizedTextCommitted = InArgs._OnTextCommitted;
+	CanDisplayResolvedText = InArgs._CanDisplayResolvedText;
+	
+	ErrorReporting = SNew(SPopupErrorText);
+	ErrorReporting->SetError(FText::GetEmpty());
+
+	bHandlingEvaluation = !ResolvedText.IsSet();
+	
+	BoundText = MakeAttributeSP(this, &SNamingTokensEditableTextBox::GetText);
+	
+	ChildSlot
+	[
+		SNew(SBorder)
+		.BorderImage(this, &SNamingTokensEditableTextBox::GetBorderImage)
+		.BorderBackgroundColor(TextBoxStyle->BackgroundColor)
+		.ForegroundColor(this, &SMultiLineEditableTextBox::GetForegroundColor)
+		.Padding(TextBoxStyle->Padding)
+		[
+			SAssignNew(Box, SHorizontalBox)
+
+			+ SHorizontalBox::Slot()
+			.HAlign(HAlign_Fill)
+			.VAlign(VAlign_Fill)
+			.AutoWidth()
+			.Padding(4, 0)
+			[
+				SNew(SImage)
+					.Image(FAppStyle::Get().GetBrush("Icons.TokenTextBox"))
+			]
+
+			+ SHorizontalBox::Slot()
+			.HAlign(HAlign_Fill)
+			.VAlign(VAlign_Fill)
+			.FillWidth(1.0f)
+			[
+				SNew(SVerticalBox)
+
+				+ SVerticalBox::Slot()
+				.HAlign(HAlign_Fill)
+				.VAlign(VAlign_Fill)
+				.FillHeight(1.0f)
+				[
+					SAssignNew(EditableText, SMultiLineEditableText)
+					.Text(BoundText)
+					.TextStyle(&TextBoxStyle->TextStyle)
+					.Marshaller(Marshaller)
+					.AllowMultiLine(InArgs._AllowMultiLine)
+					.IsReadOnly(InArgs._IsReadOnly)
+					.OverflowPolicy(ETextOverflowPolicy::Ellipsis)
+					.Margin(0.0f)
+					.OnTextChanged(this, &SNamingTokensEditableTextBox::OnEditableTextChanged)
+					.OnTextCommitted(this, &SNamingTokensEditableTextBox::OnEditableTextCommitted)
+				]
+
+				+ SVerticalBox::Slot()
+				.AutoHeight()
+				[
+					SAssignNew(HScrollBarBox, SBox)
+					.Padding(TextBoxStyle->HScrollBarPadding)
+					[
+						SAssignNew(HScrollBar, SScrollBar)
+						.Style(&TextBoxStyle->ScrollBarStyle)
+						.Orientation(Orient_Horizontal)
+						.Thickness(ScrollBarThickness)
+					]
+				]
+			]
+
+			+ SHorizontalBox::Slot()
+			.AutoWidth()
+			[
+				SAssignNew(VScrollBarBox, SBox)
+				.Padding(TextBoxStyle->VScrollBarPadding)
+				[
+					SAssignNew(VScrollBar, SScrollBar)
+					.Style(&TextBoxStyle->ScrollBarStyle)
+					.Orientation(Orient_Vertical)
+					.Thickness(ScrollBarThickness)
+				]
+			]
+
+			+ SHorizontalBox::Slot()
+			.AutoWidth()
+			.Padding(ErrorWidgetPadding)
+			[
+				ErrorReporting->AsWidget()
+			]
+		]
+	];
+
+	if (bHandlingEvaluation)
+	{
+		OnEditableTextCommitted(TokenizedText.Get(), ETextCommit::Default);
+	}
+}
+
+void SNamingTokensEditableTextBox::SetTokenizedText(const FText& InText)
+{
+	TokenizedText.Set(InText);
+	EvaluateNamingTokens();
+}
+
+const FText& SNamingTokensEditableTextBox::GetResolvedText() const
+{
+	return ResolvedText.Get();
+}
+
+const FText& SNamingTokensEditableTextBox::GetTokenizedText() const
+{
+	return TokenizedText.Get();
+}
+
+void SNamingTokensEditableTextBox::SetNormalStyle(const FEditableTextBoxStyle* InStyle)
+{
+	if (InStyle)
+	{
+		TextBoxStyle = InStyle;
+	}
+	else
+	{
+		FArguments Defaults;
+		TextBoxStyle = Defaults._Style;
+	}
+
+	check(TextBoxStyle);
+	
+	if (EditableText.IsValid())
+	{
+		EditableText->SetTextStyle(&TextBoxStyle->TextStyle);
+	}
+
+	if (Marshaller.IsValid())
+	{
+		Marshaller->SetNormalStyle(TextBoxStyle->TextStyle);
+		Marshaller->MakeDirty();
+	}
+}
+
+void SNamingTokensEditableTextBox::SetArgumentStyle(const FTextBlockStyle* InArgumentStyle)
+{
+	ArgumentStyle = InArgumentStyle;
+	
+	if (Marshaller.IsValid() && InArgumentStyle)
+	{
+		Marshaller->SetArgumentStyle(*InArgumentStyle);	
+		Marshaller->MakeDirty();
+	}
+}
+
+FText SNamingTokensEditableTextBox::GetText() const
+{
+	if (HasKeyboardFocus())
+	{
+		return TokenizedText.Get();
+	}
+	else
+	{
+		if (ResolvedText.IsSet() && CanDisplayResolvedText.Get())
+		{
+			return ResolvedText.Get();
+		}
+		else
+		{
+			return TokenizedText.Get();
+		}
+	}
+}
+
+const FSlateBrush* SNamingTokensEditableTextBox::GetBorderImage() const
+{
+	check(TextBoxStyle);
+	check(EditableText.IsValid());
+
+	if (EditableText->IsTextReadOnly())
+	{
+		return &TextBoxStyle->BackgroundImageReadOnly;
+	}
+	else if (EditableText->HasKeyboardFocus())
+	{
+		return &TextBoxStyle->BackgroundImageFocused;
+	}
+	else if (EditableText->IsHovered())
+	{
+		return &TextBoxStyle->BackgroundImageHovered;
+	}
+	else
+	{
+		return &TextBoxStyle->BackgroundImageNormal;
+	}
+}
+
+FSlateColor SNamingTokensEditableTextBox::GetForegroundColor() const
+{
+	check(TextBoxStyle);
+
+	return HasKeyboardFocus()
+		? TextBoxStyle->FocusedForegroundColor
+		: TextBoxStyle->ForegroundColor;
+}
+
+void SNamingTokensEditableTextBox::SetError(const FText& InError) const
+{
+	ErrorReporting->SetError(InError);
+}
+
+bool SNamingTokensEditableTextBox::HasKeyboardFocus() const
+{
+	// Since keyboard focus is forwarded to our editable text, we will test it instead
+	return SCompoundWidget::HasKeyboardFocus() || EditableText->HasKeyboardFocus();
+}
+
+FReply SNamingTokensEditableTextBox::OnFocusReceived(const FGeometry& MyGeometry, const FFocusEvent& InFocusEvent)
+{
+	FReply Reply = FReply::Handled();
+	if ( InFocusEvent.GetCause() != EFocusCause::Cleared )
+	{
+		// Forward keyboard focus to our editable text widget
+		Reply.SetUserFocus(EditableText.ToSharedRef(), InFocusEvent.GetCause());
+	}
+
+	return Reply;
+}
+
+FReply SNamingTokensEditableTextBox::OnKeyDown(const FGeometry& MyGeometry, const FKeyEvent& InKeyEvent)
+{
+	const FKey Key = InKeyEvent.GetKey();
+	if (Key == EKeys::Escape && EditableText->HasKeyboardFocus())
+	{
+		// Clear focus
+		return FReply::Handled().SetUserFocus(SharedThis(this), EFocusCause::Cleared);
+	}
+
+	return FReply::Unhandled();
+}
+
+bool SNamingTokensEditableTextBox::ValidateTokenizedText(const FText& InTokenizedText)
+{
+	FTextBuilder ErrorMessageBuilder;
+	FText ErrorMessage;
+
+	if (!ValidateTokenBraces(InTokenizedText, ErrorMessage))
+	{
+		ErrorMessageBuilder.AppendLine(ErrorMessage);
+	}
+
+	// Only validate tokens if there were no other errors - it's expensive
+	if (ErrorMessage.IsEmpty() && !ValidateTokenArgs(InTokenizedText, ErrorMessage))
+	{
+		ErrorMessageBuilder.AppendLine(ErrorMessage);
+	}
+
+	if (!ErrorMessageBuilder.IsEmpty())
+	{
+		ErrorMessage = ErrorMessageBuilder.ToText();
+	}
+
+	if (!ErrorMessage.IsEmpty()	&& OnValidateTokenizedText.IsBound())
+	{
+		OnValidateTokenizedText.Execute(InTokenizedText, ErrorMessage);
+	}
+
+	SetError(ErrorMessage);
+
+	// If the error message is empty, the tokenized text is valid
+	return ErrorMessage.IsEmpty();
+}
+
+bool SNamingTokensEditableTextBox::ValidateTokenBraces(const FText& InTokenizedText, FText& OutError)
+{
+	// Get number of open and close braces in string
+	int32 OpenBraceCount = 0;
+	int32 CloseBraceCount = 0;
+
+	bool bLastBraceWasOpen = false;
+
+	auto GetUnbalancedBracesMessage = []()
+	{
+		return LOCTEXT("TemplateTextUnbalancedBracingError", "An unbalanced brace was detected. Please ensure that all braces are properly closed.");
+	};
+
+	FString TextString = InTokenizedText.ToString();
+	for (const TCHAR& Char : TextString)
+	{
+		if (Char == TEXT('{'))
+		{
+			++OpenBraceCount;
+
+			// didn't close the last brace
+			if (bLastBraceWasOpen)
+			{
+				OutError = GetUnbalancedBracesMessage();
+				break;
+			}
+
+			bLastBraceWasOpen = true;
+		}
+		else if (Char == TEXT('}'))
+		{
+			++CloseBraceCount;
+			bLastBraceWasOpen = false;
+		}
+	}
+
+	if (OpenBraceCount != CloseBraceCount)
+	{
+		OutError = GetUnbalancedBracesMessage();
+	}
+
+	return OutError.IsEmpty();
+}
+
+bool SNamingTokensEditableTextBox::ValidateTokenArgs(const FText& InTokenizedText, FText& OutError)
+{
+	TConstArrayView<FString> ValidArgs = ValidArguments.Get();
+	if (ValidArgs.IsEmpty())
+	{
+		// Always valid if no provided valid args to check against
+		return true;
+	}
+
+	TArray<FString> FoundArgs;
+	if (!ParseArgs(InTokenizedText.ToString(), FoundArgs))
+	{
+		// No args in tokenized text
+		return true;
+	}
+
+	// Not particularly performant
+	for (const FString& FoundArg : FoundArgs)
+	{
+		if (!ValidArgs.ContainsByPredicate([FoundArg](const FString& InValidArg)
+		{
+			return InValidArg.Equals(FoundArg, ESearchCase::IgnoreCase);
+		}))
+		{
+			OutError = LOCTEXT("TemplateTextInvalidArgError", "An argument/token was found that doesn't match any of the provided valid argument names.");
+			return false;
+		}
+	}
+
+	return true;
+}
+
+bool SNamingTokensEditableTextBox::ParseArgs(const FString& InTokenizedTextString, TArray<FString>& OutArgs)
+{
+	OutArgs = UE::NamingTokens::Utils::GetTokenKeysFromString(InTokenizedTextString);
+	return !OutArgs.IsEmpty();
+}
+
+void SNamingTokensEditableTextBox::OnEditableTextChanged(const FText& InTokenizedText)
+{
+	ValidateTokenizedText(InTokenizedText);
+
+	OnTokenizedTextChanged.ExecuteIfBound(InTokenizedText);
+}
+
+void SNamingTokensEditableTextBox::OnEditableTextCommitted(const FText& InTokenizedText, ETextCommit::Type InCommitType)
+{
+	ValidateTokenizedText(InTokenizedText);
+
+	if (!TokenizedText.IsBound())
+	{
+		TokenizedText.Set(InTokenizedText);
+	}
+
+	OnTokenizedTextCommitted.ExecuteIfBound(InTokenizedText, InCommitType);
+
+	if (bHandlingEvaluation)
+	{
+		EvaluateNamingTokens();
+	}
+}
+
+void SNamingTokensEditableTextBox::EvaluateNamingTokens()
+{
+	check(bHandlingEvaluation);
+	
+	UNamingTokensEngineSubsystem* TokenSubsystem = GEngine->GetEngineSubsystem<UNamingTokensEngineSubsystem>();
+	const FNamingTokenResultData EvaluationData = TokenSubsystem->EvaluateTokenText(TokenizedText.Get());
+
+	ResolvedText.Set(EvaluationData.EvaluatedText);
+}
+
+END_SLATE_FUNCTION_BUILD_OPTIMIZATION
+
+#undef LOCTEXT_NAMESPACE
diff --git a/Engine/Plugins/Developer/NamingTokens/Source/NamingTokensUI/Public/NamingTokensEditableText.h b/Engine/Plugins/Developer/NamingTokens/Source/NamingTokensUI/Public/NamingTokensEditableText.h
new file mode 100644
--- /dev/null
+++ b/Engine/Plugins/Developer/NamingTokens/Source/NamingTokensUI/Public/NamingTokensEditableText.h
@@ -0,0 +1,75 @@
+// Copyright Epic Games, Inc. All Rights Reserved.
+
+#pragma once
+
+#include "Components/MultiLineEditableText.h"
+
+#include "NamingTokensEditableText.generated.h"
+
+class SNamingTokensEditableTextBox;
+struct FEditableTextBoxStyle;
+
+/**
+ * Editable text for displaying tokenized strings in either their unevaluated or resolved form.
+ */
+UCLASS(MinimalAPI, Config = NamingTokens, DefaultConfig, ClassGroup = UI, meta = (Category = "NamingTokens UI", DisplayName = "Naming Tokens Text Block", PrioritizeCategories = "Content"))
+class UNamingTokensEditableText : public UMultiLineEditableText
+{
+	GENERATED_BODY()
+
+public:
+	NAMINGTOKENSUI_API UNamingTokensEditableText();
+	
+	/** Retrieve the evaluated text. */
+	UFUNCTION(BlueprintCallable, Category=NamingTokens)
+	NAMINGTOKENSUI_API const FText& GetResolvedText() const;
+
+	/** Retrieve the raw token text. */
+	UFUNCTION(BlueprintCallable, Category=NamingTokens)
+	NAMINGTOKENSUI_API const FText& GetTokenizedText() const;
+
+	/** Sets the argument style to use for the widget. */
+	UFUNCTION(BlueprintSetter)
+	NAMINGTOKENSUI_API void SetWidgetArgumentStyle(const FTextBlockStyle& InWidgetStyle);
+
+#if WITH_EDITOR
+	virtual void PostEditChangeProperty(FPropertyChangedEvent& PropertyChangedEvent) override;
+#endif
+	
+protected:
+	//~ Begin UWidget Interface
+	virtual TSharedRef<SWidget> RebuildWidget() override;
+	virtual void SynchronizeProperties() override;
+	// End of UWidget
+
+	//~ Begin UVisual Interface
+	virtual void ReleaseSlateResources(bool bReleaseChildren) override;
+	// End UVisual Interface
+
+private:
+#if WITH_EDITOR
+	/** If we're allowed to display the resolved text. */
+	bool CanDisplayResolvedTextInEditor() const;
+#endif
+
+private:
+	/** The editable text slate widget. */
+	TSharedPtr<SNamingTokensEditableTextBox> MyNamingTokensEditableText;
+
+	/** The editable text box style we use. */
+	TSharedPtr<FEditableTextBoxStyle> EditableTextBoxStyle;
+
+	/** The font to apply to the non-resolved tokens. */
+	UPROPERTY(EditAnywhere, AdvancedDisplay, BlueprintReadWrite, BlueprintSetter=SetWidgetArgumentStyle, Category=Appearance, meta=(AllowPrivateAccess))
+	FTextBlockStyle ArgumentStyle;
+	
+	/** If this text box is configured for multiline. */
+	UPROPERTY(EditAnywhere, BlueprintReadWrite, Category=Appearance, meta=(AllowPrivateAccess))
+	bool bIsMultiline = true;
+
+#if WITH_EDITORONLY_DATA
+	/** If the resolved text should display while designing in the UMG window. If false, only the tokenized text will be displayed in the UMG designer. */
+	UPROPERTY(EditAnywhere, AdvancedDisplay, Category=Appearance)
+	bool bDisplayResolvedTextInDesigner = true;
+#endif
+};
diff --git a/Engine/Plugins/Developer/NamingTokens/Source/NamingTokensUI/Public/SNamingTokensEditableTextBox.h b/Engine/Plugins/Developer/NamingTokens/Source/NamingTokensUI/Public/SNamingTokensEditableTextBox.h
new file mode 100644
--- /dev/null
+++ b/Engine/Plugins/Developer/NamingTokens/Source/NamingTokensUI/Public/SNamingTokensEditableTextBox.h
@@ -0,0 +1,160 @@
+// Copyright Epic Games, Inc. All Rights Reserved.
+
+#pragma once
+
+#include "Containers/Array.h"
+#include "CoreMinimal.h"
+#include "Framework/SlateDelegates.h"
+#include "Misc/Attribute.h"
+#include "Templates/SharedPointer.h"
+#include "Widgets/DeclarativeSyntaxSupport.h"
+#include "Widgets/Input/SMultiLineEditableTextBox.h"
+
+class FString;
+class FText;
+class ITextLayoutMarshaller;
+class FNamingTokensStringSyntaxHighlighterMarshaller;
+
+/**
+ * An editable text box for displaying tokenized strings in either their unevaluated or resolved form.
+ */
+class SNamingTokensEditableTextBox : public SCompoundWidget
+{
+public:
+	SLATE_BEGIN_ARGS(SNamingTokensEditableTextBox)
+		: _Style(nullptr)
+		, _ArgumentStyle(nullptr)
+		, _AllowMultiLine(false)
+		, _IsReadOnly(false)
+		, _CanDisplayResolvedText(true)
+	{}
+		/** The styling of the textbox. */
+		SLATE_STYLE_ARGUMENT(FEditableTextBoxStyle, Style)
+
+		/** The styling of our arguments. */
+		SLATE_STYLE_ARGUMENT(FTextBlockStyle, ArgumentStyle)
+		
+		/** The initial text that will appear in the widget. */
+		SLATE_ATTRIBUTE(FText, Text)
+
+		/** The (optional) resolved text, displayed when the text box is not focused. If not provided, the templated text is shown. */
+		SLATE_ATTRIBUTE(FText, ResolvedText)
+
+		/** The list of available arguments to use in this template string. If Empty, any argument name is valid. */
+		SLATE_ATTRIBUTE(TConstArrayView<FString>, ValidArguments)
+
+		/** The marshaller used to get/set the raw text to/from the text layout. */
+		SLATE_ARGUMENT(TSharedPtr<ITextLayoutMarshaller>, Marshaller)
+
+		/** Whether to allow multi-line text. */
+		SLATE_ATTRIBUTE(bool, AllowMultiLine)
+
+		/** Sets whether this text box can actually be modified interactively by the user */
+		SLATE_ATTRIBUTE(bool, IsReadOnly)
+
+		/** If we're allowed to display the resolved text. */
+		SLATE_ATTRIBUTE(bool, CanDisplayResolvedText)
+
+		/** Allows custom validation. */
+		SLATE_EVENT(FOnVerifyTextChanged, OnValidateTokenizedText)
+
+		/** Called whenever the (tokenized) text is changed interactively by the user. */
+		SLATE_EVENT(FOnTextChanged, OnTextChanged)
+		
+		/** Called whenever the (tokenized) text is committed by the user. */
+		SLATE_EVENT(FOnTextCommitted, OnTextCommitted)
+
+	SLATE_END_ARGS()
+
+	/** Constructs this widget with InArgs */
+	NAMINGTOKENSUI_API void Construct(const FArguments& InArgs);
+
+	/** Retrieve the editable text widget. */
+	TSharedPtr<SMultiLineEditableText> GetEditableText() const { return EditableText; }
+
+	/** Set the current tokenized text. */
+	NAMINGTOKENSUI_API void SetTokenizedText(const FText& InText);
+	
+	/** Return the bound text attribute. */
+	const TAttribute<FText>& GetBoundText() const { return BoundText; }
+
+	/** Retrieve the resolved text. */
+	NAMINGTOKENSUI_API const FText& GetResolvedText() const;
+
+	/** Retrieve the raw tokenized text. */
+	NAMINGTOKENSUI_API const FText& GetTokenizedText() const;
+
+	/** Sets the style to use. */
+	NAMINGTOKENSUI_API void SetNormalStyle(const FEditableTextBoxStyle* InNormalStyle);
+
+	/** Sets the argument style to use. */
+	NAMINGTOKENSUI_API void SetArgumentStyle(const FTextBlockStyle* InArgumentStyle);
+	
+private:
+	/** @return Text contents based on current state */
+	FText GetText() const;
+
+	const FSlateBrush* GetBorderImage() const;
+	NAMINGTOKENSUI_API virtual FSlateColor GetForegroundColor() const override;
+
+	void SetError(const FText& InError) const;
+
+	NAMINGTOKENSUI_API virtual bool HasKeyboardFocus() const override;
+	NAMINGTOKENSUI_API virtual FReply OnFocusReceived(const FGeometry& MyGeometry, const FFocusEvent& InFocusEvent) override;
+	NAMINGTOKENSUI_API virtual FReply OnKeyDown(const FGeometry& MyGeometry, const FKeyEvent& InKeyEvent) override;
+
+	bool ValidateTokenizedText(const FText& InTokenizedText);
+	bool ValidateTokenBraces(const FText& InTokenizedText, FText& OutError);
+	bool ValidateTokenArgs(const FText& InTokenizedText, FText& OutError);
+
+	bool ParseArgs(const FString& InTokenizedTextString, TArray<FString>& OutArgs);
+
+	void OnEditableTextChanged(const FText& InTokenizedText);
+	void OnEditableTextCommitted(const FText& InTokenizedText, ETextCommit::Type InCommitType);
+
+	void EvaluateNamingTokens();
+	
+private:
+	/** Editable text widget. */
+	TSharedPtr<SMultiLineEditableText> EditableText;
+
+	/** Style shared between Editable and NonEditable text widgets. */
+	const FEditableTextBoxStyle* TextBoxStyle = nullptr;
+	/** The style to use for our argument formatting. */
+	const FTextBlockStyle* ArgumentStyle = nullptr;
+
+	TSharedPtr<SHorizontalBox> Box;
+
+	TSharedPtr<SBox> HScrollBarBox;
+	TSharedPtr<SScrollBar> HScrollBar;
+
+	TSharedPtr<SBox> VScrollBarBox;
+	TSharedPtr<SScrollBar> VScrollBar;
+
+	/** The marshaller we use. */
+	TSharedPtr<FNamingTokensStringSyntaxHighlighterMarshaller> Marshaller;
+	
+	TSharedPtr<class IErrorReportingWidget> ErrorReporting;
+
+	/** The bound text which could return either tokenized or resolved text. */
+	TAttribute<FText> BoundText;
+	
+	TAttribute<FText> TokenizedText;
+	TAttribute<FText> ResolvedText;
+
+	/** Attribute checking if we can display the resolved text. */
+	TAttribute<bool> CanDisplayResolvedText;
+	TAttribute<TConstArrayView<FString>> ValidArguments;
+	
+	/** Callback to verify tokenized text when changed. Will return an error message to denote problems. */
+    FOnVerifyTextChanged OnValidateTokenizedText;
+
+	/** Callback when tokenized text is changed. */
+	FOnTextChanged OnTokenizedTextChanged;
+
+	/** Callback when tokenized text is committed. */
+	FOnTextCommitted OnTokenizedTextCommitted;
+
+	/** If this widget is handling all token evaluation. */
+	bool bHandlingEvaluation = false;
+};
