NamingTokens: Fix static analysis error.

diff --git a/Engine/Plugins/Developer/NamingTokens/Source/NamingTokensUI/Private/NamingTokensStringSyntaxHighlighterMarshaller.cpp b/Engine/Plugins/Developer/NamingTokens/Source/NamingTokensUI/Private/NamingTokensStringSyntaxHighlighterMarshaller.cpp
--- a/Engine/Plugins/Developer/NamingTokens/Source/NamingTokensUI/Private/NamingTokensStringSyntaxHighlighterMarshaller.cpp
+++ b/Engine/Plugins/Developer/NamingTokens/Source/NamingTokensUI/Private/NamingTokensStringSyntaxHighlighterMarshaller.cpp
@@ -30,6 +30,18 @@ TSharedRef<FNamingTokensStringSyntaxHighlighterMarshaller> FNamingTokensStringSy
 	return MakeShared<FNamingTokensStringSyntaxHighlighterMarshaller>(FSyntaxTokenizer::Create(TokenizerRules), InSyntaxTextStyle);
 }
 
+const FTextBlockStyle& FNamingTokensStringSyntaxHighlighterMarshaller::GetDefaultNormalStyle()
+{
+	static FTextBlockStyle DefaultNormalStyle = FAppStyle::GetWidgetStyle<FTextBlockStyle>(TEXT("NormalText"));
+	return DefaultNormalStyle;
+}
+
+const FTextBlockStyle& FNamingTokensStringSyntaxHighlighterMarshaller::GetDefaultArgumentStyle()
+{
+	static FTextBlockStyle DefaultArgumentStyle = FAppStyle::GetWidgetStyle<FTextBlockStyle>(TEXT("RichTextBlock.BoldHighlight"));
+	return DefaultArgumentStyle;
+}
+
 void FNamingTokensStringSyntaxHighlighterMarshaller::ParseTokens(const FString& SourceString, FTextLayout& TargetTextLayout, TArray<ISyntaxTokenizer::FTokenizedLine> TokenizedLines)
 {
 	enum class EParseState : uint8
diff --git a/Engine/Plugins/Developer/NamingTokens/Source/NamingTokensUI/Private/NamingTokensStringSyntaxHighlighterMarshaller.h b/Engine/Plugins/Developer/NamingTokens/Source/NamingTokensUI/Private/NamingTokensStringSyntaxHighlighterMarshaller.h
--- a/Engine/Plugins/Developer/NamingTokens/Source/NamingTokensUI/Private/NamingTokensStringSyntaxHighlighterMarshaller.h
+++ b/Engine/Plugins/Developer/NamingTokens/Source/NamingTokensUI/Private/NamingTokensStringSyntaxHighlighterMarshaller.h
@@ -21,8 +21,8 @@ public:
 	struct FSyntaxTextStyle
 	{
 		FSyntaxTextStyle()
-			: NormalTextStyle(FAppStyle::GetWidgetStyle<FTextBlockStyle>(TEXT("NormalText")))
-			, ArgumentTextStyle(FAppStyle::GetWidgetStyle<FTextBlockStyle>("RichTextBlock.BoldHighlight"))
+			: NormalTextStyle(GetDefaultNormalStyle())
+			, ArgumentTextStyle(GetDefaultArgumentStyle())
 		{
 		}
 		
@@ -51,6 +51,12 @@ public:
 	{
 		SyntaxTextStyle.ArgumentTextStyle = InStyle;
 	}
+
+	/** Retrieve the default normal style used. */
+	static const FTextBlockStyle& GetDefaultNormalStyle();
+	
+	/** Retrieve the default argument style used. */
+	static const FTextBlockStyle& GetDefaultArgumentStyle();
 	
 protected:
 	// Allows MakeShared with private constructor
diff --git a/Engine/Plugins/Developer/NamingTokens/Source/NamingTokensUI/Private/SNamingTokensEditableTextBox.cpp b/Engine/Plugins/Developer/NamingTokens/Source/NamingTokensUI/Private/SNamingTokensEditableTextBox.cpp
--- a/Engine/Plugins/Developer/NamingTokens/Source/NamingTokensUI/Private/SNamingTokensEditableTextBox.cpp
+++ b/Engine/Plugins/Developer/NamingTokens/Source/NamingTokensUI/Private/SNamingTokensEditableTextBox.cpp
@@ -39,8 +39,9 @@ void SNamingTokensEditableTextBox::Construct(const FArguments& InArgs)
 	}
 	
 	Marshaller = FNamingTokensStringSyntaxHighlighterMarshaller::Create(
-		FNamingTokensStringSyntaxHighlighterMarshaller::FSyntaxTextStyle(TextBoxStyle->TextStyle,\
-			*InArgs._ArgumentStyle));
+		FNamingTokensStringSyntaxHighlighterMarshaller::FSyntaxTextStyle(TextBoxStyle->TextStyle,
+			InArgs._ArgumentStyle ? *InArgs._ArgumentStyle
+			: FNamingTokensStringSyntaxHighlighterMarshaller::GetDefaultArgumentStyle()));
 	
 	static const FVector2D ScrollBarThickness(9.0f, 9.0f);
 	static const FMargin ErrorWidgetPadding(3, 0); 
