NamingTokens: Fix argument font missing in packaged content.

diff --git a/Engine/Plugins/Developer/NamingTokens/Source/NamingTokensUI/Private/NamingTokensEditableText.cpp b/Engine/Plugins/Developer/NamingTokens/Source/NamingTokensUI/Private/NamingTokensEditableText.cpp
--- a/Engine/Plugins/Developer/NamingTokens/Source/NamingTokensUI/Private/NamingTokensEditableText.cpp
+++ b/Engine/Plugins/Developer/NamingTokens/Source/NamingTokensUI/Private/NamingTokensEditableText.cpp
@@ -2,11 +2,12 @@
 
 #include "NamingTokensEditableText.h"
 
+#include "NamingTokensStringSyntaxHighlighterMarshaller.h"
 #include "SNamingTokensEditableTextBox.h"
 
 UNamingTokensEditableText::UNamingTokensEditableText()
 {
-	ArgumentStyle = FAppStyle::GetWidgetStyle<FTextBlockStyle>("RichTextBlock.BoldHighlight");
+	ArgumentStyle = FNamingTokensStringSyntaxHighlighterMarshaller::GetDefaultArgumentStyle();
 }
 
 const FText& UNamingTokensEditableText::GetResolvedText() const
diff --git a/Engine/Plugins/Developer/NamingTokens/Source/NamingTokensUI/Private/NamingTokensStringSyntaxHighlighterMarshaller.cpp b/Engine/Plugins/Developer/NamingTokens/Source/NamingTokensUI/Private/NamingTokensStringSyntaxHighlighterMarshaller.cpp
--- a/Engine/Plugins/Developer/NamingTokens/Source/NamingTokensUI/Private/NamingTokensStringSyntaxHighlighterMarshaller.cpp
+++ b/Engine/Plugins/Developer/NamingTokens/Source/NamingTokensUI/Private/NamingTokensStringSyntaxHighlighterMarshaller.cpp
@@ -10,6 +10,8 @@
 #include "Framework/Text/TextLayout.h"
 #include "HAL/Platform.h"
 #include "HAL/PlatformCrt.h"
+#include "Styling/SlateStyleMacros.h"
+#include "Styling/StarshipCoreStyle.h"
 #include "Templates/UnrealTemplate.h"
 
 constexpr static const TCHAR* Operators[] = {
@@ -38,7 +40,10 @@ const FTextBlockStyle& FNamingTokensStringSyntaxHighlighterMarshaller::GetDefaul
 
 const FTextBlockStyle& FNamingTokensStringSyntaxHighlighterMarshaller::GetDefaultArgumentStyle()
 {
-	static FTextBlockStyle DefaultArgumentStyle = FAppStyle::GetWidgetStyle<FTextBlockStyle>(TEXT("RichTextBlock.BoldHighlight"));
+	static FTextBlockStyle DefaultArgumentStyle = FAppStyle::GetWidgetStyle<FTextBlockStyle>(TEXT("NormalText"));
+	DefaultArgumentStyle.SetFont(DEFAULT_FONT("Bold", FStarshipCoreStyle::RegularTextSize))
+	.SetColorAndOpacity(FLinearColor(1.0f, 1.0f, 1.0f));
+	
 	return DefaultArgumentStyle;
 }
 
