NamingTokens: Support modifying evaluation filters from the widget.
- Evaluation args can be configured in the details panel or in BP.
- Context objects can be configured in BP.
- Add OnPreEvaluateNamingTokens event to widget, allowing logic to execute right before an evaluation. Useful such as for setting contexts.

diff --git a/Engine/Plugins/Developer/NamingTokens/Source/NamingTokensUI/Private/NamingTokensEditableText.cpp b/Engine/Plugins/Developer/NamingTokens/Source/NamingTokensUI/Private/NamingTokensEditableText.cpp
--- a/Engine/Plugins/Developer/NamingTokens/Source/NamingTokensUI/Private/NamingTokensEditableText.cpp
+++ b/Engine/Plugins/Developer/NamingTokens/Source/NamingTokensUI/Private/NamingTokensEditableText.cpp
@@ -28,6 +28,11 @@ void UNamingTokensEditableText::SetWidgetArgumentStyle(const FTextBlockStyle& In
 	MyNamingTokensEditableText->SetArgumentStyle(&InWidgetStyle);
 }
 
+void UNamingTokensEditableText::SetContexts(const TArray<UObject*>& InContexts)
+{
+	Contexts = InContexts;
+}
+
 #if WITH_EDITOR
 void UNamingTokensEditableText::PostEditChangeProperty(FPropertyChangedEvent& PropertyChangedEvent)
 {
@@ -53,6 +58,13 @@ TSharedRef<SWidget> UNamingTokensEditableText::RebuildWidget()
 	EditableTextBoxStyle = MakeShared<FEditableTextBoxStyle>(FAppStyle::Get().GetWidgetStyle<FEditableTextBoxStyle>("NormalEditableTextBox"));
 	EditableTextBoxStyle->SetTextStyle(WidgetStyle);
 	
+	// We broadcast a BP event that may be setting data we need, and our widget is being constructed or reconstructed.
+	// An evaluation generally occurs on instantiation, so we fire one off without the widget to account for this case.
+	{
+		MyNamingTokensEditableText.Reset();
+		OnPreEvaluateNamingTokensCallback();
+	}
+
 	MyNamingTokensEditableText = SNew(SNamingTokensEditableTextBox)
 		.Text(GetText())
 		.IsReadOnly(GetIsReadOnly())
@@ -61,7 +73,10 @@ TSharedRef<SWidget> UNamingTokensEditableText::RebuildWidget()
 #endif
 		.AllowMultiLine(bIsMultiline)
 		.Style(&*EditableTextBoxStyle)
-		.ArgumentStyle(&ArgumentStyle);
+		.ArgumentStyle(&ArgumentStyle)
+		.NamingTokenFilterArgs(&FilterArgs)
+		.Contexts(&Contexts)
+		.OnPreEvaluateNamingTokens_UObject(this, &UNamingTokensEditableText::OnPreEvaluateNamingTokensCallback);
 
 	// Make sure our primary editable text (managed by parent) points to the editable text we create from the naming tokens slate text box.
 	MyMultiLineEditableText	= MyNamingTokensEditableText->GetEditableText();
@@ -106,6 +121,7 @@ void UNamingTokensEditableText::ReleaseSlateResources(bool bReleaseChildren)
 }
 
 #if WITH_EDITOR
+
 bool UNamingTokensEditableText::CanDisplayResolvedTextInEditor() const
 {
 	if (IsDesignTime())
@@ -115,4 +131,16 @@ bool UNamingTokensEditableText::CanDisplayResolvedTextInEditor() const
 
 	return true;
 }
-#endif
\ No newline at end of file
+
+#endif
+
+void UNamingTokensEditableText::OnPreEvaluateNamingTokensCallback()
+{
+	OnPreEvaluateNamingTokens.Broadcast();
+	
+	if (MyNamingTokensEditableText.IsValid())
+	{
+		MyNamingTokensEditableText->SetFilterArgs(&FilterArgs);
+		MyNamingTokensEditableText->SetContexts(Contexts);
+	}
+}
diff --git a/Engine/Plugins/Developer/NamingTokens/Source/NamingTokensUI/Private/SNamingTokensEditableTextBox.cpp b/Engine/Plugins/Developer/NamingTokens/Source/NamingTokensUI/Private/SNamingTokensEditableTextBox.cpp
--- a/Engine/Plugins/Developer/NamingTokens/Source/NamingTokensUI/Private/SNamingTokensEditableTextBox.cpp
+++ b/Engine/Plugins/Developer/NamingTokens/Source/NamingTokensUI/Private/SNamingTokensEditableTextBox.cpp
@@ -53,7 +53,11 @@ void SNamingTokensEditableTextBox::Construct(const FArguments& InArgs)
 	OnValidateTokenizedText = InArgs._OnValidateTokenizedText;
 	OnTokenizedTextChanged = InArgs._OnTextChanged;
 	OnTokenizedTextCommitted = InArgs._OnTextCommitted;
+	OnPreEvaluateTokens = InArgs._OnPreEvaluateNamingTokens;
 	CanDisplayResolvedText = InArgs._CanDisplayResolvedText;
+
+	SetFilterArgs(InArgs._NamingTokenFilterArgs);
+	SetContexts(InArgs._Contexts ? *InArgs._Contexts : TArray<UObject*>());
 	
 	ErrorReporting = SNew(SPopupErrorText);
 	ErrorReporting->SetError(FText::GetEmpty());
@@ -201,6 +205,23 @@ void SNamingTokensEditableTextBox::SetArgumentStyle(const FTextBlockStyle* InArg
 	}
 }
 
+void SNamingTokensEditableTextBox::SetFilterArgs(const FNamingTokenFilterArgs* InFilterArgs)
+{
+	if (InFilterArgs)
+	{
+		NamingTokenFilterArgs = MakeShared<FNamingTokenFilterArgs>(*InFilterArgs);
+	}
+	else
+	{
+		NamingTokenFilterArgs.Reset();
+	}
+}
+
+void SNamingTokensEditableTextBox::SetContexts(const TArray<UObject*>& InContexts)
+{
+	NamingTokenContexts = InContexts;
+}
+
 FText SNamingTokensEditableTextBox::GetText() const
 {
 	if (HasKeyboardFocus())
@@ -428,9 +449,12 @@ void SNamingTokensEditableTextBox::OnEditableTextCommitted(const FText& InTokeni
 void SNamingTokensEditableTextBox::EvaluateNamingTokens()
 {
 	check(bHandlingEvaluation);
+
+	OnPreEvaluateTokens.ExecuteIfBound();
 	
 	UNamingTokensEngineSubsystem* TokenSubsystem = GEngine->GetEngineSubsystem<UNamingTokensEngineSubsystem>();
-	const FNamingTokenResultData EvaluationData = TokenSubsystem->EvaluateTokenText(TokenizedText.Get());
+	const FNamingTokenResultData EvaluationData = TokenSubsystem->EvaluateTokenText(TokenizedText.Get(),
+		NamingTokenFilterArgs.IsValid() ? *NamingTokenFilterArgs : FNamingTokenFilterArgs(), NamingTokenContexts);
 
 	ResolvedText.Set(EvaluationData.EvaluatedText);
 }
diff --git a/Engine/Plugins/Developer/NamingTokens/Source/NamingTokensUI/Public/NamingTokensEditableText.h b/Engine/Plugins/Developer/NamingTokens/Source/NamingTokensUI/Public/NamingTokensEditableText.h
--- a/Engine/Plugins/Developer/NamingTokens/Source/NamingTokensUI/Public/NamingTokensEditableText.h
+++ b/Engine/Plugins/Developer/NamingTokens/Source/NamingTokensUI/Public/NamingTokensEditableText.h
@@ -3,6 +3,7 @@
 #pragma once
 
 #include "Components/MultiLineEditableText.h"
+#include "NamingTokensEngineSubsystem.h"
 
 #include "NamingTokensEditableText.generated.h"
 
@@ -18,6 +19,8 @@ class UNamingTokensEditableText : public UMultiLineEditableText
 	GENERATED_BODY()
 
 public:
+	DECLARE_DYNAMIC_MULTICAST_DELEGATE(FOnPreEvaluateNamingTokens);
+	
 	NAMINGTOKENSUI_API UNamingTokensEditableText();
 	
 	/** Retrieve the evaluated text. */
@@ -32,6 +35,10 @@ public:
 	UFUNCTION(BlueprintSetter)
 	NAMINGTOKENSUI_API void SetWidgetArgumentStyle(const FTextBlockStyle& InWidgetStyle);
 
+	/** Set the contexts to use during evaluation. */
+	UFUNCTION(BlueprintSetter)
+	NAMINGTOKENSUI_API void SetContexts(const TArray<UObject*>& InContexts);
+	
 #if WITH_EDITOR
 	virtual void PostEditChangeProperty(FPropertyChangedEvent& PropertyChangedEvent) override;
 #endif
@@ -45,13 +52,26 @@ protected:
 	//~ Begin UVisual Interface
 	virtual void ReleaseSlateResources(bool bReleaseChildren) override;
 	// End UVisual Interface
-
+	
 private:
 #if WITH_EDITOR
 	/** If we're allowed to display the resolved text. */
 	bool CanDisplayResolvedTextInEditor() const;
 #endif
 
+	/** Callback when we are evaluating tokens. */
+	void OnPreEvaluateNamingTokensCallback();
+
+protected:
+	/** Set naming token filter args to use during token evaluation. */
+	UPROPERTY(EditAnywhere, BlueprintReadWrite, Category=NamingTokens)
+	FNamingTokenFilterArgs FilterArgs;
+
+public:
+	/** Called before naming token evaluation. */
+	UPROPERTY(BlueprintAssignable, Category="NamingTokens|Event")
+	FOnPreEvaluateNamingTokens OnPreEvaluateNamingTokens;
+	
 private:
 	/** The editable text slate widget. */
 	TSharedPtr<SNamingTokensEditableTextBox> MyNamingTokensEditableText;
@@ -59,6 +79,10 @@ private:
 	/** The editable text box style we use. */
 	TSharedPtr<FEditableTextBoxStyle> EditableTextBoxStyle;
 
+	/** Contexts to include during our evaluation. */
+	UPROPERTY(BlueprintReadWrite, BlueprintSetter=SetContexts, Category=NamingTokens, meta=(AllowPrivateAccess))
+	TArray<TObjectPtr<UObject>> Contexts;
+	
 	/** The font to apply to the non-resolved tokens. */
 	UPROPERTY(EditAnywhere, AdvancedDisplay, BlueprintReadWrite, BlueprintSetter=SetWidgetArgumentStyle, Category=Appearance, meta=(AllowPrivateAccess))
 	FTextBlockStyle ArgumentStyle;
diff --git a/Engine/Plugins/Developer/NamingTokens/Source/NamingTokensUI/Public/SNamingTokensEditableTextBox.h b/Engine/Plugins/Developer/NamingTokens/Source/NamingTokensUI/Public/SNamingTokensEditableTextBox.h
--- a/Engine/Plugins/Developer/NamingTokens/Source/NamingTokensUI/Public/SNamingTokensEditableTextBox.h
+++ b/Engine/Plugins/Developer/NamingTokens/Source/NamingTokensUI/Public/SNamingTokensEditableTextBox.h
@@ -14,6 +14,7 @@ class FString;
 class FText;
 class ITextLayoutMarshaller;
 class FNamingTokensStringSyntaxHighlighterMarshaller;
+struct FNamingTokenFilterArgs;
 
 /**
  * An editable text box for displaying tokenized strings in either their unevaluated or resolved form.
@@ -27,6 +28,7 @@ public:
 		, _AllowMultiLine(false)
 		, _IsReadOnly(false)
 		, _CanDisplayResolvedText(true)
+		, _NamingTokenFilterArgs(nullptr)
 	{}
 		/** The styling of the textbox. */
 		SLATE_STYLE_ARGUMENT(FEditableTextBoxStyle, Style)
@@ -64,6 +66,15 @@ public:
 		/** Called whenever the (tokenized) text is committed by the user. */
 		SLATE_EVENT(FOnTextCommitted, OnTextCommitted)
 
+		/** Called prior to Naming Token evaluation. Only fires if this widget is handling evaluations. */
+		SLATE_EVENT(FSimpleDelegate, OnPreEvaluateNamingTokens)
+
+		/** Filter args to apply during evaluation. */
+		SLATE_ARGUMENT(const FNamingTokenFilterArgs*, NamingTokenFilterArgs)
+
+		/** Contexts to provide to evaluation. */
+		SLATE_ARGUMENT(const TArray<TObjectPtr<UObject>>*, Contexts)
+
 	SLATE_END_ARGS()
 
 	/** Constructs this widget with InArgs */
@@ -89,6 +100,12 @@ public:
 
 	/** Sets the argument style to use. */
 	NAMINGTOKENSUI_API void SetArgumentStyle(const FTextBlockStyle* InArgumentStyle);
+
+	/** Provide filter args to use during evaluation. Passing null clears existing. */
+	NAMINGTOKENSUI_API void SetFilterArgs(const FNamingTokenFilterArgs* InFilterArgs);
+
+	/** Provide contexts to naming token evaluation. */
+	NAMINGTOKENSUI_API void SetContexts(const TArray<UObject*>& InContexts);
 	
 private:
 	/** @return Text contents based on current state */
@@ -115,6 +132,12 @@ private:
 	void EvaluateNamingTokens();
 	
 private:
+	/** Filter args to use during evaluation. */
+	TSharedPtr<FNamingTokenFilterArgs> NamingTokenFilterArgs;
+
+	/** Contexts to provide to evaluation. */
+	TArray<TObjectPtr<UObject>> NamingTokenContexts;
+	
 	/** Editable text widget. */
 	TSharedPtr<SMultiLineEditableText> EditableText;
 
@@ -155,6 +178,9 @@ private:
 	/** Callback when tokenized text is committed. */
 	FOnTextCommitted OnTokenizedTextCommitted;
 
+	/** Callback before we evaluate tokens. */
+	FSimpleDelegate OnPreEvaluateTokens;
+
 	/** If this widget is handling all token evaluation. */
 	bool bHandlingEvaluation = false;
 };
