NamingTokens: UMG widget utilizes all of underlying Naming Token slate widget, instead of just the text.
- Add token icon, error message, border image, background style options to UMG widget.
- Improve slate token image display on multiline entries.

diff --git a/Engine/Plugins/Developer/NamingTokens/Source/NamingTokensUI/Private/NamingTokensEditableText.cpp b/Engine/Plugins/Developer/NamingTokens/Source/NamingTokensUI/Private/NamingTokensEditableText.cpp
--- a/Engine/Plugins/Developer/NamingTokens/Source/NamingTokensUI/Private/NamingTokensEditableText.cpp
+++ b/Engine/Plugins/Developer/NamingTokens/Source/NamingTokensUI/Private/NamingTokensEditableText.cpp
@@ -8,6 +8,7 @@
 UNamingTokensEditableText::UNamingTokensEditableText()
 {
 	ArgumentStyle = FNamingTokensStringSyntaxHighlighterMarshaller::GetDefaultArgumentStyle();
+	BackgroundColor = FLinearColor(0, 0, 0, 0);
 }
 
 const FText& UNamingTokensEditableText::GetResolvedText() const
@@ -24,8 +25,21 @@ const FText& UNamingTokensEditableText::GetTokenizedText() const
 
 void UNamingTokensEditableText::SetWidgetArgumentStyle(const FTextBlockStyle& InWidgetStyle)
 {
-	check(MyNamingTokensEditableText.IsValid());
-	MyNamingTokensEditableText->SetArgumentStyle(&InWidgetStyle);
+	ArgumentStyle = InWidgetStyle;
+	if (MyNamingTokensEditableText.IsValid())
+	{
+		MyNamingTokensEditableText->SetArgumentStyle(&InWidgetStyle);
+	}
+}
+
+void UNamingTokensEditableText::SetBackgroundColor(const FSlateColor& InColor)
+{
+	BackgroundColor = InColor;
+	if (MyNamingTokensEditableText.IsValid() && EditableTextBoxStyle.IsValid())
+	{
+		EditableTextBoxStyle->SetBackgroundColor(InColor);
+		MyNamingTokensEditableText->SetNormalStyle(EditableTextBoxStyle.Get());
+	}
 }
 
 void UNamingTokensEditableText::SetContexts(const TArray<UObject*>& InContexts)
@@ -33,6 +47,36 @@ void UNamingTokensEditableText::SetContexts(const TArray<UObject*>& InContexts)
 	Contexts = InContexts;
 }
 
+void UNamingTokensEditableText::SetDisplayTokenIcon(bool bValue)
+{
+	bDisplayTokenIcon = bValue;
+}
+
+void UNamingTokensEditableText::SetDisplayErrorMessage(bool bValue)
+{
+	bDisplayErrorMessage = bValue;
+}
+
+void UNamingTokensEditableText::SetDisplayBorderImage(bool bValue)
+{
+	bDisplayBorderImage = bValue;
+}
+
+bool UNamingTokensEditableText::GetDisplayTokenIcon() const
+{
+	return bDisplayTokenIcon;
+}
+
+bool UNamingTokensEditableText::GetDisplayErrorMessage() const
+{
+	return bDisplayErrorMessage;
+}
+
+bool UNamingTokensEditableText::GetDisplayBorderImage() const
+{
+	return bDisplayBorderImage;
+}
+
 #if WITH_EDITOR
 void UNamingTokensEditableText::PostEditChangeProperty(FPropertyChangedEvent& PropertyChangedEvent)
 {
@@ -57,6 +101,7 @@ TSharedRef<SWidget> UNamingTokensEditableText::RebuildWidget()
 	// We only use text box style, but slate needs editable text box style.
 	EditableTextBoxStyle = MakeShared<FEditableTextBoxStyle>(FAppStyle::Get().GetWidgetStyle<FEditableTextBoxStyle>("NormalEditableTextBox"));
 	EditableTextBoxStyle->SetTextStyle(WidgetStyle);
+	EditableTextBoxStyle->SetBackgroundColor(BackgroundColor);
 	
 	// We broadcast a BP event that may be setting data we need, and our widget is being constructed or reconstructed.
 	// An evaluation generally occurs on instantiation, so we fire one off without the widget to account for this case.
@@ -76,12 +121,15 @@ TSharedRef<SWidget> UNamingTokensEditableText::RebuildWidget()
 		.ArgumentStyle(&ArgumentStyle)
 		.NamingTokenFilterArgs(&FilterArgs)
 		.Contexts(&Contexts)
+		.DisplayTokenIcon_UObject(this, &UNamingTokensEditableText::GetDisplayTokenIcon)
+		.DisplayErrorMessage_UObject(this, &UNamingTokensEditableText::GetDisplayErrorMessage)
+		.DisplayBorderImage_UObject(this, &UNamingTokensEditableText::GetDisplayBorderImage)
 		.OnPreEvaluateNamingTokens_UObject(this, &UNamingTokensEditableText::OnPreEvaluateNamingTokensCallback);
 
 	// Make sure our primary editable text (managed by parent) points to the editable text we create from the naming tokens slate text box.
 	MyMultiLineEditableText	= MyNamingTokensEditableText->GetEditableText();
 
-	return MyMultiLineEditableText.ToSharedRef();
+	return MyNamingTokensEditableText.ToSharedRef();
 }
 
 void UNamingTokensEditableText::SynchronizeProperties()
diff --git a/Engine/Plugins/Developer/NamingTokens/Source/NamingTokensUI/Private/SNamingTokensEditableTextBox.cpp b/Engine/Plugins/Developer/NamingTokens/Source/NamingTokensUI/Private/SNamingTokensEditableTextBox.cpp
--- a/Engine/Plugins/Developer/NamingTokens/Source/NamingTokensUI/Private/SNamingTokensEditableTextBox.cpp
+++ b/Engine/Plugins/Developer/NamingTokens/Source/NamingTokensUI/Private/SNamingTokensEditableTextBox.cpp
@@ -56,6 +56,10 @@ void SNamingTokensEditableTextBox::Construct(const FArguments& InArgs)
 	OnPreEvaluateTokens = InArgs._OnPreEvaluateNamingTokens;
 	CanDisplayResolvedText = InArgs._CanDisplayResolvedText;
 
+	DisplayTokenIcon = InArgs._DisplayTokenIcon;
+	DisplayErrorMessage = InArgs._DisplayErrorMessage;
+	DisplayBorderImage = InArgs._DisplayBorderImage;
+	
 	SetFilterArgs(InArgs._NamingTokenFilterArgs);
 	SetContexts(InArgs._Contexts ? *InArgs._Contexts : TArray<UObject*>());
 	
@@ -70,20 +74,25 @@ void SNamingTokensEditableTextBox::Construct(const FArguments& InArgs)
 	[
 		SNew(SBorder)
 		.BorderImage(this, &SNamingTokensEditableTextBox::GetBorderImage)
-		.BorderBackgroundColor(TextBoxStyle->BackgroundColor)
+		.BorderBackgroundColor(this, &SNamingTokensEditableTextBox::GetBorderBackgroundColor)
 		.ForegroundColor(this, &SMultiLineEditableTextBox::GetForegroundColor)
+		.VAlign(VAlign_Center)
 		.Padding(TextBoxStyle->Padding)
 		[
 			SAssignNew(Box, SHorizontalBox)
 
 			+ SHorizontalBox::Slot()
 			.HAlign(HAlign_Fill)
-			.VAlign(VAlign_Fill)
+			.VAlign(VAlign_Center)
 			.AutoWidth()
 			.Padding(4, 0)
 			[
 				SNew(SImage)
 					.Image(FAppStyle::Get().GetBrush("Icons.TokenTextBox"))
+					.Visibility_Lambda([this]()
+					{
+						return DisplayTokenIcon.Get(false) ? EVisibility::Visible : EVisibility::Collapsed;
+					})
 			]
 
 			+ SHorizontalBox::Slot()
@@ -141,7 +150,14 @@ void SNamingTokensEditableTextBox::Construct(const FArguments& InArgs)
 			.AutoWidth()
 			.Padding(ErrorWidgetPadding)
 			[
-				ErrorReporting->AsWidget()
+				SNew(SBox)
+				.Visibility_Lambda([this]()
+				{
+					return DisplayErrorMessage.Get(false) ? EVisibility::Visible : EVisibility::Collapsed;
+				})
+				[
+					ErrorReporting->AsWidget()
+				]
 			]
 		]
 	];
@@ -246,6 +262,11 @@ const FSlateBrush* SNamingTokensEditableTextBox::GetBorderImage() const
 	check(TextBoxStyle);
 	check(EditableText.IsValid());
 
+	if (DisplayBorderImage.IsSet() && !DisplayBorderImage.Get())
+	{
+		return nullptr;
+	}
+
 	if (EditableText->IsTextReadOnly())
 	{
 		return &TextBoxStyle->BackgroundImageReadOnly;
@@ -264,6 +285,12 @@ const FSlateBrush* SNamingTokensEditableTextBox::GetBorderImage() const
 	}
 }
 
+FSlateColor SNamingTokensEditableTextBox::GetBorderBackgroundColor() const
+{
+	check(TextBoxStyle);
+	return TextBoxStyle->BackgroundColor;
+}
+
 FSlateColor SNamingTokensEditableTextBox::GetForegroundColor() const
 {
 	check(TextBoxStyle);
diff --git a/Engine/Plugins/Developer/NamingTokens/Source/NamingTokensUI/Public/NamingTokensEditableText.h b/Engine/Plugins/Developer/NamingTokens/Source/NamingTokensUI/Public/NamingTokensEditableText.h
--- a/Engine/Plugins/Developer/NamingTokens/Source/NamingTokensUI/Public/NamingTokensEditableText.h
+++ b/Engine/Plugins/Developer/NamingTokens/Source/NamingTokensUI/Public/NamingTokensEditableText.h
@@ -35,9 +35,34 @@ public:
 	UFUNCTION(BlueprintSetter)
 	NAMINGTOKENSUI_API void SetWidgetArgumentStyle(const FTextBlockStyle& InWidgetStyle);
 
+	/** Set the background color. */
+	UFUNCTION(BlueprintSetter)
+	NAMINGTOKENSUI_API void SetBackgroundColor(const FSlateColor& InColor);
+	
 	/** Set the contexts to use during evaluation. */
 	UFUNCTION(BlueprintSetter)
 	NAMINGTOKENSUI_API void SetContexts(const TArray<UObject*>& InContexts);
+
+	/** Set whether the token icon is displayed. */
+	UFUNCTION(BlueprintSetter)
+	NAMINGTOKENSUI_API void SetDisplayTokenIcon(bool bValue);
+
+	/** Set whether error messages are displayed for token formatting. */
+	UFUNCTION(BlueprintSetter)
+	NAMINGTOKENSUI_API void SetDisplayErrorMessage(bool bValue);
+
+	/** Set the background color of the border image. */
+	UFUNCTION(BlueprintSetter)
+	NAMINGTOKENSUI_API void SetDisplayBorderImage(bool bValue);
+
+	/** If the token icon should be displayed. */
+	NAMINGTOKENSUI_API bool GetDisplayTokenIcon() const;
+
+	/** If the error icon should be displayed. */
+	NAMINGTOKENSUI_API bool GetDisplayErrorMessage() const;
+
+	/** If the border image should be displayed. */
+	NAMINGTOKENSUI_API bool GetDisplayBorderImage() const;
 	
 #if WITH_EDITOR
 	virtual void PostEditChangeProperty(FPropertyChangedEvent& PropertyChangedEvent) override;
@@ -61,12 +86,42 @@ private:
 
 	/** Callback when we are evaluating tokens. */
 	void OnPreEvaluateNamingTokensCallback();
-
+	
 protected:
 	/** Set naming token filter args to use during token evaluation. */
 	UPROPERTY(EditAnywhere, BlueprintReadWrite, Category=NamingTokens)
 	FNamingTokenFilterArgs FilterArgs;
 
+	/** If this text box is configured for multiline. */
+	UPROPERTY(EditAnywhere, BlueprintReadWrite, Category=Appearance)
+	bool bIsMultiline = true;
+
+	/** Display the token icon in the text box. */
+	UPROPERTY(EditAnywhere, AdvancedDisplay, BlueprintReadWrite, BlueprintSetter=SetDisplayTokenIcon, Category=Appearance)
+	bool bDisplayTokenIcon = false;
+
+	/** Display an error message when tokens aren't properly formatted. */
+	UPROPERTY(EditAnywhere, AdvancedDisplay, BlueprintReadWrite, BlueprintSetter=SetDisplayErrorMessage, Category=Appearance)
+	bool bDisplayErrorMessage = false;
+
+	/** Display the border image. */
+	UPROPERTY(EditAnywhere, AdvancedDisplay, BlueprintReadWrite, BlueprintSetter=SetDisplayBorderImage, Category=Appearance)
+	bool bDisplayBorderImage = false;
+	
+	/** The font to apply to the non-resolved tokens. */
+	UPROPERTY(EditAnywhere, AdvancedDisplay, BlueprintReadWrite, BlueprintSetter=SetWidgetArgumentStyle, Category=Appearance)
+	FTextBlockStyle ArgumentStyle;
+	
+	/** Set the background color of the border image. */
+	UPROPERTY(EditAnywhere, AdvancedDisplay, BlueprintReadWrite, BlueprintSetter=SetBackgroundColor, Category=Appearance)
+	FSlateColor BackgroundColor;
+	
+#if WITH_EDITORONLY_DATA
+	/** If the resolved text should display while designing in the UMG window. If false, only the tokenized text will be displayed in the UMG designer. */
+	UPROPERTY(EditAnywhere, AdvancedDisplay, Category=Appearance)
+	bool bDisplayResolvedTextInDesigner = true;
+#endif
+	
 public:
 	/** Called before naming token evaluation. */
 	UPROPERTY(BlueprintAssignable, Category="NamingTokens|Event")
@@ -82,18 +137,4 @@ private:
 	/** Contexts to include during our evaluation. */
 	UPROPERTY(BlueprintReadWrite, BlueprintSetter=SetContexts, Category=NamingTokens, meta=(AllowPrivateAccess))
 	TArray<TObjectPtr<UObject>> Contexts;
-	
-	/** The font to apply to the non-resolved tokens. */
-	UPROPERTY(EditAnywhere, AdvancedDisplay, BlueprintReadWrite, BlueprintSetter=SetWidgetArgumentStyle, Category=Appearance, meta=(AllowPrivateAccess))
-	FTextBlockStyle ArgumentStyle;
-	
-	/** If this text box is configured for multiline. */
-	UPROPERTY(EditAnywhere, BlueprintReadWrite, Category=Appearance, meta=(AllowPrivateAccess))
-	bool bIsMultiline = true;
-
-#if WITH_EDITORONLY_DATA
-	/** If the resolved text should display while designing in the UMG window. If false, only the tokenized text will be displayed in the UMG designer. */
-	UPROPERTY(EditAnywhere, AdvancedDisplay, Category=Appearance)
-	bool bDisplayResolvedTextInDesigner = true;
-#endif
 };
diff --git a/Engine/Plugins/Developer/NamingTokens/Source/NamingTokensUI/Public/SNamingTokensEditableTextBox.h b/Engine/Plugins/Developer/NamingTokens/Source/NamingTokensUI/Public/SNamingTokensEditableTextBox.h
--- a/Engine/Plugins/Developer/NamingTokens/Source/NamingTokensUI/Public/SNamingTokensEditableTextBox.h
+++ b/Engine/Plugins/Developer/NamingTokens/Source/NamingTokensUI/Public/SNamingTokensEditableTextBox.h
@@ -28,6 +28,9 @@ public:
 		, _AllowMultiLine(false)
 		, _IsReadOnly(false)
 		, _CanDisplayResolvedText(true)
+		, _DisplayTokenIcon(true)
+		, _DisplayErrorMessage(true)
+		, _DisplayBorderImage(true)
 		, _NamingTokenFilterArgs(nullptr)
 	{}
 		/** The styling of the textbox. */
@@ -57,6 +60,15 @@ public:
 		/** If we're allowed to display the resolved text. */
 		SLATE_ATTRIBUTE(bool, CanDisplayResolvedText)
 
+		/** If we can display the default token icon. */
+		SLATE_ATTRIBUTE(bool, DisplayTokenIcon)
+
+		/** If we display error messages. */
+		SLATE_ATTRIBUTE(bool, DisplayErrorMessage)
+
+		/** If we should use the border image. */
+		SLATE_ATTRIBUTE(bool, DisplayBorderImage)
+
 		/** Allows custom validation. */
 		SLATE_EVENT(FOnVerifyTextChanged, OnValidateTokenizedText)
 
@@ -112,6 +124,8 @@ private:
 	FText GetText() const;
 
 	const FSlateBrush* GetBorderImage() const;
+	/** The background color for the border. */
+	virtual FSlateColor GetBorderBackgroundColor() const;
 	NAMINGTOKENSUI_API virtual FSlateColor GetForegroundColor() const override;
 
 	void SetError(const FText& InError) const;
@@ -165,6 +179,15 @@ private:
 	TAttribute<FText> TokenizedText;
 	TAttribute<FText> ResolvedText;
 
+	/** If we are displaying the token icon. */
+	TAttribute<bool> DisplayTokenIcon;
+
+	/** If error messages are displayed. */
+	TAttribute<bool> DisplayErrorMessage;
+	
+	/** If the border image is used. */
+	TAttribute<bool> DisplayBorderImage;
+	
 	/** Attribute checking if we can display the resolved text. */
 	TAttribute<bool> CanDisplayResolvedText;
 	TAttribute<TConstArrayView<FString>> ValidArguments;
