NamingTokens: Default bIsMultiline on UMG widget to false.

diff --git a/Engine/Plugins/Developer/NamingTokens/Source/NamingTokensUI/Public/NamingTokensEditableText.h b/Engine/Plugins/Developer/NamingTokens/Source/NamingTokensUI/Public/NamingTokensEditableText.h
--- a/Engine/Plugins/Developer/NamingTokens/Source/NamingTokensUI/Public/NamingTokensEditableText.h
+++ b/Engine/Plugins/Developer/NamingTokens/Source/NamingTokensUI/Public/NamingTokensEditableText.h
@@ -94,7 +94,7 @@ protected:
 
 	/** If this text box is configured for multiline. */
 	UPROPERTY(EditAnywhere, BlueprintReadWrite, Category=Appearance)
-	bool bIsMultiline = true;
+	bool bIsMultiline = false;
 
 	/** Display the token icon in the text box. */
 	UPROPERTY(EditAnywhere, AdvancedDisplay, BlueprintReadWrite, BlueprintSetter=SetDisplayTokenIcon, Category=Appearance)
