NamingTokens: Support case-insensitive namespace lookup.

bForceCaseSensitive is now recognized for namespaces to force case-sensitive checks.

Registering namespaces is unchanged, and still requires a unique name regardless of case.

diff --git a/Engine/Plugins/Developer/NamingTokens/Source/NamingTokens/Private/NamingTokens.cpp b/Engine/Plugins/Developer/NamingTokens/Source/NamingTokens/Private/NamingTokens.cpp
--- a/Engine/Plugins/Developer/NamingTokens/Source/NamingTokens/Private/NamingTokens.cpp
+++ b/Engine/Plugins/Developer/NamingTokens/Source/NamingTokens/Private/NamingTokens.cpp
@@ -184,7 +184,10 @@ FNamingTokenResultData UNamingTokens::EvaluateTokenText(const FText& InTokenText
 			
 			if (SelectedToken)
 			{
-				TokenValue = ProcessToken(*SelectedToken, TokenKey, FormattedText);
+				// If case-sensitive is set, provide an empty string so ProcessToken will look up the exact namespace for this token. Otherwise,
+				// we use what the user provided.
+				const FString NamespaceToCheck = InEvaluationData.bForceCaseSensitive ? FString() : TokenNamespace;
+				TokenValue = ProcessToken(*SelectedToken, NamespaceToCheck, TokenKey, FormattedText);
 				ValueData.bWasEvaluated = true;
 			}
 			else
@@ -296,7 +299,7 @@ void UNamingTokens::OnPostEvaluate_Implementation()
 {
 }
 
-FText UNamingTokens::ProcessToken(const FNamingTokenData& InToken, const FString& InUserTokenString, FText& InOutFormattedText)
+FText UNamingTokens::ProcessToken(const FNamingTokenData& InToken, const FString& InUserProvidedNamespace, const FString& InUserTokenString, FText& InOutFormattedText)
 {
 	const FString FormattedToken = *UE::NamingTokens::Utils::CreateFormattedToken(InToken);
 
@@ -324,7 +327,8 @@ FText UNamingTokens::ProcessToken(const FNamingTokenData& InToken, const FString
 	// Replace both {token} and {namespace:token} in the formatted text.
 	FFormatNamedArguments Args;
 	Args.Add(InUserTokenString, ProcessedToken);
-	Args.Add(UE::NamingTokens::Utils::CombineNamespaceAndTokenKey(*GetNamespace(), *InUserTokenString), ProcessedToken);
+	Args.Add(UE::NamingTokens::Utils::CombineNamespaceAndTokenKey(InUserProvidedNamespace.IsEmpty()
+		? *GetNamespace() : *InUserProvidedNamespace, *InUserTokenString), ProcessedToken);
 
 	InOutFormattedText = FText::Format(InOutFormattedText, Args);
 	return ProcessedToken;
diff --git a/Engine/Plugins/Developer/NamingTokens/Source/NamingTokens/Public/NamingTokens.h b/Engine/Plugins/Developer/NamingTokens/Source/NamingTokens/Public/NamingTokens.h
--- a/Engine/Plugins/Developer/NamingTokens/Source/NamingTokens/Public/NamingTokens.h
+++ b/Engine/Plugins/Developer/NamingTokens/Source/NamingTokens/Public/NamingTokens.h
@@ -154,7 +154,7 @@ private:
 	FOnPostEvaluate OnPostEvaluateEvent;
 	
 	/** Process a token if is defined and return evaluated text. */
-	UE_API FText ProcessToken(const FNamingTokenData& InToken, const FString& InUserTokenString, FText& InOutFormattedText);
+	UE_API FText ProcessToken(const FNamingTokenData& InToken, const FString& InUserProvidedNamespace, const FString& InUserTokenString, FText& InOutFormattedText);
 	
 	/** Return a blueprint function for a token processor if it exists. */
 	UE_API UFunction* FindBlueprintFunctionForToken(const FNamingTokenData& InTokenData) const;
diff --git a/Engine/Plugins/Developer/NamingTokens/Source/NamingTokens/Public/Utils/NamingTokenUtils.h b/Engine/Plugins/Developer/NamingTokens/Source/NamingTokens/Public/Utils/NamingTokenUtils.h
--- a/Engine/Plugins/Developer/NamingTokens/Source/NamingTokens/Public/Utils/NamingTokenUtils.h
+++ b/Engine/Plugins/Developer/NamingTokens/Source/NamingTokens/Public/Utils/NamingTokenUtils.h
@@ -28,14 +28,14 @@ namespace UE::NamingTokens::Utils
 	
 	/**
 	 * Return the namespace from a token key, or an empty string.
-	 * @param InTokenKey The token key shouldn't have any additional formatting beyond the namespace, eg "namespace.token"
+	 * @param InTokenKey The token key shouldn't have any additional formatting beyond the namespace, eg "namespace:token"
 	 * @return The namespace, or an empty string.
 	 */
 	NAMINGTOKENS_API FString GetNamespaceFromTokenKey(const FString& InTokenKey);
 	
 	/**
 	 * Remove the namespace from the given token key.
-	 * @param InTokenKey The token key shouldn't have any additional formatting beyond the namespace, eg "namespace.token"
+	 * @param InTokenKey The token key shouldn't have any additional formatting beyond the namespace, eg "namespace:token"
 	 * @return The token key without the namespace, or the InTokenKey if there was no namespace.
 	 */
 	NAMINGTOKENS_API FString RemoveNamespaceFromTokenKey(const FString& InTokenKey);
