NamingTokens: Add auto complete suggestion dropdown to UMG widget.

diff --git a/Engine/Plugins/Animation/LiveLinkHub/Source/LiveLinkHub/Private/Config/LiveLinkHubTemplateTokens.cpp b/Engine/Plugins/Animation/LiveLinkHub/Source/LiveLinkHub/Private/Config/LiveLinkHubTemplateTokens.cpp
--- a/Engine/Plugins/Animation/LiveLinkHub/Source/LiveLinkHub/Private/Config/LiveLinkHubTemplateTokens.cpp
+++ b/Engine/Plugins/Animation/LiveLinkHub/Source/LiveLinkHub/Private/Config/LiveLinkHubTemplateTokens.cpp
@@ -12,6 +12,7 @@
 ULiveLinkHubNamingTokens::ULiveLinkHubNamingTokens()
 {
 	Namespace = ILiveLinkHubModule::GetLiveLinkHubNamingTokensNamespace();
+	NamespaceDisplayName = LOCTEXT("NamespaceDisplayName", "Live Link Hub");
 }
 
 void ULiveLinkHubNamingTokens::OnCreateDefaultTokens(TArray<FNamingTokenData>& Tokens)
diff --git a/Engine/Plugins/Developer/NamingTokens/Source/NamingTokens/Public/NamingTokens.h b/Engine/Plugins/Developer/NamingTokens/Source/NamingTokens/Public/NamingTokens.h
--- a/Engine/Plugins/Developer/NamingTokens/Source/NamingTokens/Public/NamingTokens.h
+++ b/Engine/Plugins/Developer/NamingTokens/Source/NamingTokens/Public/NamingTokens.h
@@ -95,6 +95,9 @@ public:
 	const FString& GetNamespace() const { return Namespace; }
 	/** The name of the namespace property. */
 	static FName GetNamespacePropertyName() { return GET_MEMBER_NAME_CHECKED(UNamingTokens, Namespace); }
+
+	/** Retrieve the namespace's friendly display name. */
+	const FText& GetNamespaceDisplayName() const { return NamespaceDisplayName; }
 	
 	/** Delegate when pre evaluate is called. */
 	DECLARE_MULTICAST_DELEGATE_OneParam(FOnPreEvaluate, const FNamingTokensEvaluationData& InEvaluationData);
@@ -182,6 +185,12 @@ protected:
 	 */
 	UPROPERTY(EditDefaultsOnly, AssetRegistrySearchable, Category = "NamingTokens")
 	FString Namespace;
+
+	/**
+	 * The full display name of the namespace to use in UI and filtering.
+	 */
+	UPROPERTY(EditDefaultsOnly, AssetRegistrySearchable, Category = "NamingTokens")
+	FText NamespaceDisplayName;
 	
 #if WITH_EDITORONLY_DATA
 private:
diff --git a/Engine/Plugins/Developer/NamingTokens/Source/NamingTokensUI/Private/NamingTokensEditableText.cpp b/Engine/Plugins/Developer/NamingTokens/Source/NamingTokensUI/Private/NamingTokensEditableText.cpp
--- a/Engine/Plugins/Developer/NamingTokens/Source/NamingTokensUI/Private/NamingTokensEditableText.cpp
+++ b/Engine/Plugins/Developer/NamingTokens/Source/NamingTokensUI/Private/NamingTokensEditableText.cpp
@@ -121,6 +121,7 @@ TSharedRef<SWidget> UNamingTokensEditableText::RebuildWidget()
 		.ArgumentStyle(&ArgumentStyle)
 		.NamingTokenFilterArgs(&FilterArgs)
 		.Contexts(&Contexts)
+		.EnableSuggestionDropdown(bEnableSuggestionDropdown)
 		.DisplayTokenIcon_UObject(this, &UNamingTokensEditableText::GetDisplayTokenIcon)
 		.DisplayErrorMessage_UObject(this, &UNamingTokensEditableText::GetDisplayErrorMessage)
 		.DisplayBorderImage_UObject(this, &UNamingTokensEditableText::GetDisplayBorderImage)
@@ -148,6 +149,7 @@ void UNamingTokensEditableText::SynchronizeProperties()
 	EditableTextBoxStyle->SetTextStyle(WidgetStyle);
 	MyNamingTokensEditableText->SetNormalStyle(&*EditableTextBoxStyle);
 	MyNamingTokensEditableText->SetArgumentStyle(&ArgumentStyle);
+	MyNamingTokensEditableText->SetEnableSuggestionDropdown(bEnableSuggestionDropdown);
 	
 	MyMultiLineEditableText->SetText(MyNamingTokensEditableText->GetBoundText());
 	MyMultiLineEditableText->SetHintText(HintTextBinding);
diff --git a/Engine/Plugins/Developer/NamingTokens/Source/NamingTokensUI/Private/SNamingTokensDataTreeView.cpp b/Engine/Plugins/Developer/NamingTokens/Source/NamingTokensUI/Private/SNamingTokensDataTreeView.cpp
new file mode 100644
--- /dev/null
+++ b/Engine/Plugins/Developer/NamingTokens/Source/NamingTokensUI/Private/SNamingTokensDataTreeView.cpp
@@ -0,0 +1,410 @@
+// Copyright Epic Games, Inc. All Rights Reserved.
+
+#include "SNamingTokensDataTreeView.h"
+
+#include "NamingTokens.h"
+#include "NamingTokensEngineSubsystem.h"
+
+#include "Engine/Engine.h"
+
+#define LOCTEXT_NAMESPACE "SNamingTokenDataTreeView"
+
+void SNamingTokenDataTreeView::Construct(const FArguments& InArgs)
+{
+	STreeView::Construct(InArgs);
+}
+
+void SNamingTokenDataTreeView::SetExpansionStateFromItems(const TArray<FNamingTokenDataTreeItemPtr>& InTreeItems)
+{
+	for (const FNamingTokenDataTreeItemPtr& TreeItem: InTreeItems)
+	{
+		SetItemExpansion(TreeItem, TreeItem->ShouldItemBeExpanded());
+		SetExpansionStateFromItems(TreeItem->ChildItems);
+	}
+}
+
+void SNamingTokenDataTreeViewRow::Construct(const FArguments& InArgs, const TSharedPtr<SNamingTokenDataTreeView>& TreeView)
+{
+	Item = InArgs._Item;
+
+	const FSuperRowType::FArguments SuperArgs;
+	SMultiColumnTableRow::Construct(SuperArgs, TreeView.ToSharedRef());
+}
+
+TSharedRef<SWidget> SNamingTokenDataTreeViewRow::GenerateWidgetForColumn(const FName& ColumnName)
+{
+	const FNamingTokenDataTreeItemPtr ItemPtr = Item.Pin();
+	if (ItemPtr.IsValid())
+	{
+		const TSharedPtr<SBorder> Border = SNew(SBorder)
+			.BorderImage(FAppStyle::Get().GetBrush("NoBorder"));
+		
+		if (ItemPtr->NamingTokenData.IsValid() && ColumnName == SNamingTokenDataTreeViewWidget::PrimaryColumnName())
+		{
+			Border->SetContent(
+				SNew(SVerticalBox)
+				+ SVerticalBox::Slot()
+				.AutoHeight()
+				.Padding(16.f, 2.f)
+				[
+					SNew(STextBlock)
+					.Text(FText::FromString(ItemPtr->NamingTokenData->TokenKey))
+				]);
+		}
+		else if (!ItemPtr->Namespace.IsEmpty())
+		{
+			const FText NamespaceDisplayName = !ItemPtr->NamespaceDisplayName.IsEmpty() ?
+				FText::Format(LOCTEXT("NamespaceDisplayName", "{0} ({1}:)"),
+				ItemPtr->NamespaceDisplayName, FText::FromString(ItemPtr->Namespace))
+					: FText::FromString(ItemPtr->Namespace);
+			
+			Border->SetContent(SNew(SHorizontalBox)
+				+ SHorizontalBox::Slot()
+				.AutoWidth()
+				.Padding(0.f, 0.f, 0.f, 0.f)
+				.VAlign(VAlign_Center)
+				[
+					SNew(SExpanderArrow, SharedThis(this))
+					.StyleSet(&FAppStyle::Get())
+					.IndentAmount(5)
+				]
+				+ SHorizontalBox::Slot()
+				.FillWidth(1.0f)
+				.VAlign(VAlign_Center)
+				[
+					SNew(STextBlock)
+					.Font(FCoreStyle::Get().GetFontStyle("ExpandableArea.TitleFont"))
+					.Text(NamespaceDisplayName)
+				]);
+		}
+
+		return Border.ToSharedRef();
+	}
+
+	return SNullWidget::NullWidget;
+}
+
+SNamingTokenDataTreeViewWidget::~SNamingTokenDataTreeViewWidget()
+{
+}
+
+void SNamingTokenDataTreeViewWidget::Construct(const FArguments& InArgs)
+{
+	ItemSelectionChangedDelegate = InArgs._OnSelectionChanged;
+	ItemSelectedDelegate = InArgs._OnItemSelected;
+	
+	TSharedPtr<SHeaderRow> HeaderRow;
+	// Setup the columns.
+	{
+		SAssignNew(HeaderRow, SHeaderRow);
+
+		SHeaderRow::FColumn::FArguments ColumnArgs;
+		ColumnArgs.ColumnId(PrimaryColumnName());
+		ColumnArgs.DefaultLabel(LOCTEXT("ItemLabel_HeaderText", "Primary"));
+
+		// Don't draw the column header.
+		HeaderRow->SetVisibility(EVisibility::Collapsed);
+		HeaderRow->AddColumn(ColumnArgs);
+	}
+
+	PopulateTreeItems();
+
+	ChildSlot
+	[
+		SNew(SBox)
+		.MaxDesiredHeight(250.f)
+		.Padding(1.f)
+		[
+			SAssignNew(TreeView, SNamingTokenDataTreeView)
+			.SelectionMode(ESelectionMode::Single)
+			.TreeItemsSource(&FilteredRootTreeItems)
+			.HeaderRow(HeaderRow)
+			.OnGenerateRow(this, &SNamingTokenDataTreeViewWidget::OnGenerateRowForTree)
+			.OnGetChildren(this, &SNamingTokenDataTreeViewWidget::OnGetChildrenForTree)
+			.OnExpansionChanged(this, &SNamingTokenDataTreeViewWidget::OnItemExpansionChanged)
+			.OnSelectionChanged(this, &SNamingTokenDataTreeViewWidget::OnTreeViewItemSelectionChanged)
+			.OnMouseButtonDoubleClick(this, &SNamingTokenDataTreeViewWidget::OnTreeViewItemDoubleClicked)
+		]
+	];
+
+	FilterTreeItems(FString());
+}
+
+void SNamingTokenDataTreeViewWidget::PopulateTreeItems()
+{
+	RootTreeItems.Reset();
+	if (GEngine)
+	{
+		if (const UNamingTokensEngineSubsystem* NamingTokensSubsystem = GEngine->GetEngineSubsystem<UNamingTokensEngineSubsystem>())
+		{
+			const TArray<FString> Namespaces = NamingTokensSubsystem->GetAllNamespaces();
+			RootTreeItems.Reserve(Namespaces.Num());
+
+			int32 TotalCount = 0;
+			
+			for (const FString& Namespace : Namespaces)
+			{
+				FNamingTokenDataTreeItemPtr NamespaceTreeItem = MakeShared<FNamingTokenDataTreeItem>();
+				NamespaceTreeItem->Namespace = Namespace;
+				NamespaceTreeItem->ItemIndex = TotalCount++;
+				NamespaceTreeItem->bIsGlobal = NamingTokensSubsystem->IsGlobalNamespaceRegistered(Namespace);
+				
+				if (const UNamingTokens* NamingTokens = NamingTokensSubsystem->GetNamingTokens(Namespace))
+				{
+					NamespaceTreeItem->NamespaceDisplayName = NamingTokens->GetNamespaceDisplayName();
+					
+					const TArray<FNamingTokenData> Tokens = NamingTokens->GetAllTokens();
+					if (Tokens.Num() == 0)
+					{
+						// Don't display empty namespaces.
+						continue;
+					}
+					
+					NamespaceTreeItem->ChildItems.Reserve(Tokens.Num());
+					
+					for (const FNamingTokenData& Token : Tokens)
+					{
+						FNamingTokenDataTreeItemPtr TokenTreeItem = MakeShared<FNamingTokenDataTreeItem>();
+						TokenTreeItem->NamingTokenData = MakeShared<FNamingTokenData>(Token);
+						TokenTreeItem->Namespace = Namespace;
+						TokenTreeItem->bIsGlobal = NamespaceTreeItem->bIsGlobal;
+						TokenTreeItem->ItemIndex = TotalCount++;
+						NamespaceTreeItem->ChildItems.Add(TokenTreeItem);
+					}
+				}
+				
+				RootTreeItems.Add(NamespaceTreeItem);
+			}
+		}
+	}
+}
+
+void SNamingTokenDataTreeViewWidget::FilterTreeItems(const FString& InFilterText)
+{
+	bool bSelectingNamespace = false;
+	if (LastFilterText == InFilterText && FilteredRootTreeItems.Num() > 0)
+	{
+		if (const FNamingTokenDataTreeItemPtr CurrentlySelectedItem = GetSelectedItem())
+		{
+			// This is used at the end to skip selecting the namespace again if that's all we're showing.
+			// We may need to move this logic elsewhere in case filtering is called when selecting for some reason.
+			bSelectingNamespace = InFilterText == CurrentlySelectedItem->Namespace && CurrentlySelectedItem->IsNamespace();
+		}
+	}
+	LastFilterText = InFilterText;
+		
+	FString CompleteOrPartialNamespace = UE::NamingTokens::Utils::GetNamespaceFromTokenKey(InFilterText);
+	FString PartialToken = UE::NamingTokens::Utils::RemoveNamespaceFromTokenKey(InFilterText);
+
+	const bool bSearchNamespaceAndToken = CompleteOrPartialNamespace.IsEmpty();
+	if (bSearchNamespaceAndToken)
+	{
+		// Can't discern either part, assume searching for namespaces.
+		
+		CompleteOrPartialNamespace = InFilterText;
+		PartialToken = InFilterText;
+	}
+
+	// We have a complete namespace we require.
+	const bool bExactNamespaceMatch = !bSearchNamespaceAndToken;
+	
+	FilteredRootTreeItems.Reset();
+
+	// The best matched namespace or token index found while filtering.
+	TTuple<int32, float> BestMatchedIndex = { 0, 0.f };
+	
+	int32 TotalCount = 0;
+	
+	for (const FNamingTokenDataTreeItemPtr& TreeItem : RootTreeItems)
+	{
+		// Copy our item so we can modify the contained children.
+		const FNamingTokenDataTreeItemPtr FilteredNamespace = MakeShared<FNamingTokenDataTreeItem>(*TreeItem);
+		FilteredNamespace->ChildItems.Reset();
+		
+		// Check if the namespace matches.
+		const float NamespaceMatchPercent = TreeItem->MatchesFilter(CompleteOrPartialNamespace, bExactNamespaceMatch);
+		if (bSearchNamespaceAndToken || NamespaceMatchPercent > 0.f)
+		{
+			int32 IterationCount = TotalCount;
+			FilteredNamespace->ItemIndex = IterationCount++;
+
+			TTuple<int32, float> HighestTokenIndexScore = TTuple<int32, float>(0, 0.f);
+			for (const FNamingTokenDataTreeItemPtr& ChildItem : TreeItem->ChildItems)
+			{
+				// The token itself.
+				const FNamingTokenDataTreeItemPtr FilteredChild = MakeShared<FNamingTokenDataTreeItem>(*ChildItem);
+				const float TokenMatchPercent = FilteredChild->MatchesFilter(PartialToken);
+				if (TokenMatchPercent > 0.f || NamespaceMatchPercent > 0.f)
+				{
+					FilteredChild->ItemIndex = IterationCount++;
+					FilteredNamespace->ChildItems.Add(FilteredChild);
+
+					if (TokenMatchPercent > HighestTokenIndexScore.Get<1>())
+					{
+						HighestTokenIndexScore = TTuple<int32, float>(FilteredChild->ItemIndex, TokenMatchPercent);
+					}
+				}
+			}
+
+			// Don't include namespace if it or its tokens don't match.
+			if (NamespaceMatchPercent || FilteredNamespace->ChildItems.Num() > 0)
+			{
+				// Only set persistent count when we know were committing this result.
+				TotalCount = IterationCount;
+				FilteredRootTreeItems.Add(FilteredNamespace);
+
+				// Determine the best match between namespace and token results.
+				TTuple<int32, float> HighestOverallScore = { 0, 0.f };
+
+				bool bUsingNamespaceScore = false;
+				// Prioritize namespace first unless we're already at an exact match, because then we're searching for tokens
+				// inside that namespace.
+				if (!bExactNamespaceMatch && NamespaceMatchPercent >= HighestTokenIndexScore.Get<1>() && NamespaceMatchPercent > 0.f)
+				{
+					HighestOverallScore = TTuple<int32, float>(FilteredNamespace->ItemIndex, NamespaceMatchPercent);
+					bUsingNamespaceScore = true;
+				}
+				else
+				{
+					HighestOverallScore = HighestTokenIndexScore;
+				}
+
+				// Record the best match overall so far.
+				if (HighestOverallScore.Get<1>() > BestMatchedIndex.Get<1>()
+					// Namespaces should take priority.
+					|| (bUsingNamespaceScore && HighestOverallScore.Get<1>() >= 0.75f && HighestOverallScore.Get<1>() >= BestMatchedIndex.Get<1>()))
+				{
+					BestMatchedIndex = MoveTemp(HighestOverallScore);
+				}
+			}
+			
+			if (bExactNamespaceMatch)
+			{
+				// We've already found an exact match to our namespace, no need to continue searching namespaces.
+				break;
+			}
+		}
+	}
+	
+	TreeView->RequestListRefresh();
+	TreeView->SetExpansionStateFromItems(FilteredRootTreeItems);
+
+	if (FilteredRootTreeItems.Num() > 0)
+	{
+		SetItemFromIndex(bSelectingNamespace ? 1 : BestMatchedIndex.Get<0>());
+	}
+}
+
+void SNamingTokenDataTreeViewWidget::ForwardKeyEventForNavigation(const FKeyEvent& InKeyEvent)
+{
+	check(TreeView.IsValid());
+
+	if (InKeyEvent.GetKey() == EKeys::Tab || InKeyEvent.GetKey() == EKeys::Enter)
+	{
+		FinalizeSelection(GetSelectedItem());
+		return;
+	}
+	
+	const EUINavigation Direction = FSlateApplicationBase::Get().GetNavigationDirectionFromKey(InKeyEvent);
+	if (Direction != EUINavigation::Invalid)
+	{
+		int32 NewIndex = CurrentSelectionIndex;
+		if (Direction == EUINavigation::Down)
+		{
+			NewIndex++;
+		}
+		else if (Direction == EUINavigation::Up)
+		{
+			if (CurrentSelectionIndex > 0)
+			{
+				NewIndex--;
+			}
+		}
+
+		SetItemFromIndex(NewIndex);
+	}
+}
+
+FNamingTokenDataTreeItemPtr SNamingTokenDataTreeViewWidget::GetSelectedItem() const
+{
+	return GetItemFromIndex(CurrentSelectionIndex, FilteredRootTreeItems);
+}
+
+FNamingTokenDataTreeItemPtr SNamingTokenDataTreeViewWidget::GetItemFromIndex(int32 InIndex) const
+{
+	return GetItemFromIndex(InIndex, FilteredRootTreeItems);
+}
+
+FNamingTokenDataTreeItemPtr SNamingTokenDataTreeViewWidget::GetItemFromIndex(int32 InIndex, const TArray<FNamingTokenDataTreeItemPtr>& InItems) const
+{
+	for (int32 Idx = 0; Idx < InItems.Num(); ++Idx)
+	{
+		const FNamingTokenDataTreeItemPtr& TreeItem = InItems[Idx];
+		if (TreeItem->ItemIndex == InIndex)
+		{
+			return TreeItem;
+		}
+
+		if (TreeItem->ChildItems.Num() > 0)
+		{
+			const FNamingTokenDataTreeItemPtr& ChildItem = GetItemFromIndex(InIndex, TreeItem->ChildItems);
+			if (ChildItem.IsValid())
+			{
+				return ChildItem;
+			}
+		}
+	}
+
+	return nullptr;
+}
+
+void SNamingTokenDataTreeViewWidget::SetItemFromIndex(int32 InIndex)
+{
+	if (const FNamingTokenDataTreeItemPtr Item = GetItemFromIndex(InIndex))
+	{
+		if (const FNamingTokenDataTreeItemPtr CurrentSelection = GetSelectedItem())
+		{
+			TreeView->SetItemSelection(CurrentSelection, false);
+		}
+		
+		CurrentSelectionIndex = InIndex;
+		TreeView->SetItemSelection(Item, true);
+		if (!TreeView->IsItemVisible(Item))
+		{
+			TreeView->RequestScrollIntoView(Item);
+		}
+	}
+}
+
+void SNamingTokenDataTreeViewWidget::FinalizeSelection(FNamingTokenDataTreeItemPtr SelectedItem)
+{
+	ItemSelectedDelegate.ExecuteIfBound(SelectedItem);
+}
+
+TSharedRef<ITableRow> SNamingTokenDataTreeViewWidget::OnGenerateRowForTree(FNamingTokenDataTreeItemPtr Item,
+                                                                           const TSharedRef<STableViewBase>& OwnerTable)
+{
+	return SNew(SNamingTokenDataTreeViewRow, TreeView).Item(Item);
+}
+
+void SNamingTokenDataTreeViewWidget::OnGetChildrenForTree(FNamingTokenDataTreeItemPtr InParent, TArray<FNamingTokenDataTreeItemPtr>& OutChildren)
+{
+	OutChildren.Append(InParent->ChildItems);
+}
+
+void SNamingTokenDataTreeViewWidget::OnItemExpansionChanged(FNamingTokenDataTreeItemPtr TreeItem, bool bIsExpanded) const
+{
+	
+}
+
+void SNamingTokenDataTreeViewWidget::OnTreeViewItemSelectionChanged(FNamingTokenDataTreeItemPtr SelectedItem, ESelectInfo::Type SelectInfo)
+{
+	ItemSelectionChangedDelegate.ExecuteIfBound(SelectedItem, SelectInfo);
+}
+
+void SNamingTokenDataTreeViewWidget::OnTreeViewItemDoubleClicked(FNamingTokenDataTreeItemPtr SelectedItem)
+{
+	FinalizeSelection(SelectedItem);
+}
+
+#undef LOCTEXT_NAMESPACE
diff --git a/Engine/Plugins/Developer/NamingTokens/Source/NamingTokensUI/Private/SNamingTokensDataTreeView.h b/Engine/Plugins/Developer/NamingTokens/Source/NamingTokensUI/Private/SNamingTokensDataTreeView.h
new file mode 100644
--- /dev/null
+++ b/Engine/Plugins/Developer/NamingTokens/Source/NamingTokensUI/Private/SNamingTokensDataTreeView.h
@@ -0,0 +1,220 @@
+// Copyright Epic Games, Inc. All Rights Reserved.
+
+#pragma once
+
+#include "NamingTokenData.h"
+#include "Utils/NamingTokenUtils.h"
+
+#include "Widgets/Views/STreeView.h"
+
+/**
+ * Represents naming token items in our tree view.
+ */
+struct FNamingTokenDataTreeItem : TSharedFromThis<FNamingTokenDataTreeItem>
+{
+	/** The token data, null on namespace entries. */
+	TSharedPtr<FNamingTokenData> NamingTokenData;
+
+	/** The namespace for this token. May be set regardless of entry type. */
+	FString Namespace;
+	/** The friendly display name for the namespace. */
+	FText NamespaceDisplayName;
+
+	/** Children of a namespace. */
+	TArray<TSharedPtr<FNamingTokenDataTreeItem>> ChildItems;
+	/** The parent owning this entry. */
+	TWeakPtr<FNamingTokenDataTreeItem> Parent;
+
+	/** The flattened item index in the tree. */
+	int32 ItemIndex = INDEX_NONE;
+
+	/** If this entry is globally registered in the naming tokens subsystem. */
+	bool bIsGlobal = false;
+	
+	bool ShouldItemBeExpanded() const { return true; }
+
+	/** Checks if this item matches a filter returning a simple percentage score on the match. */
+	float MatchesFilter(const FString& FilterText, bool bIsExactNamespaceMatch = false) const
+	{
+		if (FilterText.IsEmpty())
+		{
+			return 0.5f;
+		}
+
+		if (IsNamespace())
+		{
+			if (Namespace == FilterText) // Exact match requires true namespace being typed out...
+			{
+				return 1.f;
+			}
+			
+			if (!bIsExactNamespaceMatch &&
+				(Namespace.StartsWith(FilterText, ESearchCase::IgnoreCase) // ...Otherwise, we can check partial match and the display name.
+				|| NamespaceDisplayName.ToString().StartsWith(FilterText, ESearchCase::IgnoreCase)))
+			{
+				return 0.75f;
+			}
+
+			return 0.f;
+		}
+
+		check(NamingTokenData.IsValid());
+
+		if (NamingTokenData->TokenKey == FilterText)
+		{
+			return 1.f;
+		}
+		
+		if (NamingTokenData->TokenKey.StartsWith(FilterText, ESearchCase::IgnoreCase)
+			|| NamingTokenData->DisplayName.ToString().StartsWith(FilterText, ESearchCase::IgnoreCase))
+		{
+			return 0.75f;
+		}
+
+		return 0.f;
+	}
+
+	/** True iff this is considered a namespace entry. */
+	bool IsNamespace() const
+	{
+		return !NamingTokenData.IsValid();
+	}
+	
+	/** Convert the item to the best text representation. */
+	FString ToString(bool bFullyQualified) const
+	{
+		if (NamingTokenData.IsValid())
+		{
+			return bFullyQualified ?
+				UE::NamingTokens::Utils::CombineNamespaceAndTokenKey(Namespace, NamingTokenData->TokenKey)
+					: NamingTokenData->TokenKey;
+		}
+		return Namespace;
+	}
+	
+	bool operator==(const FNamingTokenDataTreeItem& Other) const
+	{
+		return NamingTokenData == Other.NamingTokenData && Namespace == Other.Namespace;
+	}
+};
+
+typedef TSharedPtr<FNamingTokenDataTreeItem> FNamingTokenDataTreeItemPtr;
+
+/**
+ * A tree view for naming token filtering.
+ */
+class SNamingTokenDataTreeView : public STreeView<FNamingTokenDataTreeItemPtr>
+{
+public:
+	void Construct(const FArguments& InArgs);
+
+	/** Recursively set expansion state of tree view to match items. */
+	void SetExpansionStateFromItems(const TArray<FNamingTokenDataTreeItemPtr>& InTreeItems);
+};
+
+/**
+ * Each row of the tree view.
+ */
+class SNamingTokenDataTreeViewRow final : public SMultiColumnTableRow<FNamingTokenDataTreeItemPtr>
+{
+public:
+	SLATE_BEGIN_ARGS(SNamingTokenDataTreeViewRow) {}
+
+	/** The list item for this row. */
+	SLATE_ARGUMENT(FNamingTokenDataTreeItemPtr, Item)
+	SLATE_END_ARGS()
+	
+	void Construct(const FArguments& InArgs, const TSharedPtr<SNamingTokenDataTreeView>& TreeView);
+
+	// ~Begin SMultiColumnTableRow interface
+	virtual TSharedRef<SWidget> GenerateWidgetForColumn(const FName& ColumnName) override;
+	// ~End SMultiColumnTableRow interface
+
+private:
+	/** The item associated with this row of data. */
+	TWeakPtr<FNamingTokenDataTreeItem> Item;
+};
+
+/**
+ * Widget for displaying namespaces and tokens to choose from.
+ */
+class SNamingTokenDataTreeViewWidget final : public SCompoundWidget
+{
+	using FOnSelectionChanged = STreeView<FNamingTokenDataTreeItemPtr>::FOnSelectionChanged;
+
+	DECLARE_DELEGATE_OneParam(FOnItemSelected, FNamingTokenDataTreeItemPtr);
+	
+public:
+	SLATE_BEGIN_ARGS(SNamingTokenDataTreeViewWidget) {}
+		SLATE_EVENT(FOnSelectionChanged, OnSelectionChanged)
+		SLATE_EVENT(FOnItemSelected, OnItemSelected)
+	SLATE_END_ARGS()
+
+	virtual ~SNamingTokenDataTreeViewWidget() override;
+	
+	void Construct(const FArguments& InArgs);
+
+	/** Populate all possible tree items */
+	void PopulateTreeItems();
+
+	/** Filter down to specific tokens and namespaces. */
+	void FilterTreeItems(const FString& InFilterText);
+	
+	/** Provide a key event from the user to determine a navigation path, highlighting an item. */
+	void ForwardKeyEventForNavigation(const FKeyEvent& InKeyEvent);
+
+	/** Retrieve the currently selected item. */
+	FNamingTokenDataTreeItemPtr GetSelectedItem() const;
+
+	/** Retrieve an item given its absolute index. */
+	FNamingTokenDataTreeItemPtr GetItemFromIndex(int32 InIndex) const;
+private:
+	FNamingTokenDataTreeItemPtr GetItemFromIndex(int32 InIndex, const TArray<FNamingTokenDataTreeItemPtr>& InItems) const;
+
+	/** Set the currently selected item given an absolute index. */
+	void SetItemFromIndex(int32 InIndex);
+
+	/** Officially complete and report the selection. Called from tab or double-click. */
+	void FinalizeSelection(FNamingTokenDataTreeItemPtr SelectedItem);
+	
+public:
+	static FName PrimaryColumnName() { return "Primary"; }
+	
+private:
+	/** Called by STreeView for each row being generated. */
+	TSharedRef<ITableRow> OnGenerateRowForTree(FNamingTokenDataTreeItemPtr Item, const TSharedRef<STableViewBase>& OwnerTable);
+	
+	/** Called by STreeView to get child items for the specified parent item. */
+	void OnGetChildrenForTree(FNamingTokenDataTreeItemPtr InParent, TArray<FNamingTokenDataTreeItemPtr>& OutChildren);
+	
+	/** Called when an item in the tree has been collapsed or expanded. */
+	void OnItemExpansionChanged(FNamingTokenDataTreeItemPtr TreeItem, bool bIsExpanded) const;
+
+	/** Called when the tree view changes selection. */
+	void OnTreeViewItemSelectionChanged(FNamingTokenDataTreeItemPtr SelectedItem, ESelectInfo::Type SelectInfo);
+
+	/** Called when an item in the tree view is double-clicked. */
+	void OnTreeViewItemDoubleClicked(FNamingTokenDataTreeItemPtr SelectedItem);
+	
+private:
+	/** Our tree view widget. */
+	TSharedPtr<SNamingTokenDataTreeView> TreeView;
+
+	/** The top most items on the tree view without a filter. */
+	TArray<FNamingTokenDataTreeItemPtr> RootTreeItems;
+
+	/** Items filtered in a search. */
+	TArray<FNamingTokenDataTreeItemPtr> FilteredRootTreeItems;
+
+	/** The last filter applied. */
+	FString LastFilterText;
+	
+	/** Index of tree item currently selected. */
+	int32 CurrentSelectionIndex = 0;
+
+	/** Delegate for when selection changes. */
+	FOnSelectionChanged ItemSelectionChangedDelegate;
+
+	/** Delegate for when an item is fully selected. */
+	FOnItemSelected ItemSelectedDelegate;
+};
\ No newline at end of file
diff --git a/Engine/Plugins/Developer/NamingTokens/Source/NamingTokensUI/Private/SNamingTokensEditableTextBox.cpp b/Engine/Plugins/Developer/NamingTokens/Source/NamingTokensUI/Private/SNamingTokensEditableTextBox.cpp
--- a/Engine/Plugins/Developer/NamingTokens/Source/NamingTokensUI/Private/SNamingTokensEditableTextBox.cpp
+++ b/Engine/Plugins/Developer/NamingTokens/Source/NamingTokensUI/Private/SNamingTokensEditableTextBox.cpp
@@ -4,11 +4,14 @@
 
 #include "NamingTokensEngineSubsystem.h"
 #include "NamingTokensStringSyntaxHighlighterMarshaller.h"
+#include "SNamingTokensDataTreeView.h"
 #include "Utils/NamingTokenUtils.h"
 
 #include "Engine/Engine.h"
+#include "Internationalization/Regex.h"
 #include "Internationalization/Text.h"
 #include "Layout/Margin.h"
+#include "Layout/WidgetPath.h"
 #include "SlateOptMacros.h"
 #include "Styling/ISlateStyle.h"
 #include "Styling/SlateTypes.h"
@@ -59,6 +62,9 @@ void SNamingTokensEditableTextBox::Construct(const FArguments& InArgs)
 	DisplayTokenIcon = InArgs._DisplayTokenIcon;
 	DisplayErrorMessage = InArgs._DisplayErrorMessage;
 	DisplayBorderImage = InArgs._DisplayBorderImage;
+
+	bEnableSuggestionDropdown = InArgs._EnableSuggestionDropdown;
+	bFullyQualifyGlobalTokenSuggestions = InArgs._FullyQualifyGlobalTokenSuggestions;
 	
 	SetFilterArgs(InArgs._NamingTokenFilterArgs);
 	SetContexts(InArgs._Contexts ? *InArgs._Contexts : TArray<UObject*>());
@@ -107,16 +113,28 @@ void SNamingTokensEditableTextBox::Construct(const FArguments& InArgs)
 				.VAlign(VAlign_Fill)
 				.FillHeight(1.0f)
 				[
-					SAssignNew(EditableText, SMultiLineEditableText)
-					.Text(BoundText)
-					.TextStyle(&TextBoxStyle->TextStyle)
-					.Marshaller(Marshaller)
-					.AllowMultiLine(InArgs._AllowMultiLine)
-					.IsReadOnly(InArgs._IsReadOnly)
-					.OverflowPolicy(ETextOverflowPolicy::Ellipsis)
-					.Margin(0.0f)
-					.OnTextChanged(this, &SNamingTokensEditableTextBox::OnEditableTextChanged)
-					.OnTextCommitted(this, &SNamingTokensEditableTextBox::OnEditableTextCommitted)
+					SAssignNew(DropdownAnchor, SMenuAnchor)
+					.Method(EPopupMethod::CreateNewWindow)
+					.Placement(MenuPlacement_ComboBox)
+					// The drop down filter menu.
+					.OnGetMenuContent(this, &SNamingTokensEditableTextBox::OnGetDropdownContent)
+					.Visibility(EVisibility::Visible)
+					.Content()
+					[
+						SAssignNew(EditableText, SMultiLineEditableText)
+						.Text(BoundText)
+						.TextStyle(&TextBoxStyle->TextStyle)
+						.Marshaller(Marshaller)
+						.AllowMultiLine(InArgs._AllowMultiLine)
+						.IsReadOnly(InArgs._IsReadOnly)
+						.OverflowPolicy(ETextOverflowPolicy::Ellipsis)
+						.Margin(0.0f)
+						.OnTextChanged(this, &SNamingTokensEditableTextBox::OnEditableTextChanged)
+						.OnTextCommitted(this, &SNamingTokensEditableTextBox::OnEditableTextCommitted)
+						.OnKeyDownHandler(this, &SNamingTokensEditableTextBox::OnEditableTextKeyDown)
+						.OnKeyCharHandler(this, &SNamingTokensEditableTextBox::OnEditableTextKeyChar)
+						.OnCursorMoved(this, &SNamingTokensEditableTextBox::OnEditableTextCursorMoved)
+					]
 				]
 
 				+ SVerticalBox::Slot()
@@ -161,7 +179,7 @@ void SNamingTokensEditableTextBox::Construct(const FArguments& InArgs)
 			]
 		]
 	];
-
+	
 	if (bHandlingEvaluation)
 	{
 		OnEditableTextCommitted(TokenizedText.Get(), ETextCommit::Default);
@@ -221,6 +239,11 @@ void SNamingTokensEditableTextBox::SetArgumentStyle(const FTextBlockStyle* InArg
 	}
 }
 
+void SNamingTokensEditableTextBox::SetEnableSuggestionDropdown(bool bInEnableSuggestionDropdown)
+{
+	bEnableSuggestionDropdown = bInEnableSuggestionDropdown;
+}
+
 void SNamingTokensEditableTextBox::SetFilterArgs(const FNamingTokenFilterArgs* InFilterArgs)
 {
 	if (InFilterArgs)
@@ -308,7 +331,7 @@ void SNamingTokensEditableTextBox::SetError(const FText& InError) const
 bool SNamingTokensEditableTextBox::HasKeyboardFocus() const
 {
 	// Since keyboard focus is forwarded to our editable text, we will test it instead
-	return SCompoundWidget::HasKeyboardFocus() || EditableText->HasKeyboardFocus();
+	return SCompoundWidget::HasKeyboardFocus() || EditableText->HasKeyboardFocus() || (DropdownAnchor.IsValid() && DropdownAnchor->IsOpen());
 }
 
 FReply SNamingTokensEditableTextBox::OnFocusReceived(const FGeometry& MyGeometry, const FFocusEvent& InFocusEvent)
@@ -323,6 +346,16 @@ FReply SNamingTokensEditableTextBox::OnFocusReceived(const FGeometry& MyGeometry
 	return Reply;
 }
 
+void SNamingTokensEditableTextBox::OnFocusChanging(const FWeakWidgetPath& PreviousFocusPath, const FWidgetPath& NewWidgetPath,
+	const FFocusEvent& InFocusEvent)
+{
+	// We want to close the dropdown if whatever focus path doesn't contain our widgets.
+	if (FilterWidget.IsValid() && !NewWidgetPath.ContainsWidget(this) && !NewWidgetPath.ContainsWidget(FilterWidget.Get()))
+	{
+		CloseSuggestionDropdown();
+	}
+}
+
 FReply SNamingTokensEditableTextBox::OnKeyDown(const FGeometry& MyGeometry, const FKeyEvent& InKeyEvent)
 {
 	const FKey Key = InKeyEvent.GetKey();
@@ -331,6 +364,12 @@ FReply SNamingTokensEditableTextBox::OnKeyDown(const FGeometry& MyGeometry, cons
 		// Clear focus
 		return FReply::Handled().SetUserFocus(SharedThis(this), EFocusCause::Cleared);
 	}
+	if (DropdownAnchor->IsOpen() && FilterWidget.IsValid())
+	{
+		FilterWidget->ForwardKeyEventForNavigation(InKeyEvent);
+		// Set focus back to the editable text so the user can keep typing.
+		return FReply::Handled().SetUserFocus(EditableText.ToSharedRef(), EFocusCause::OtherWidgetLostFocus);
+	}
 
 	return FReply::Unhandled();
 }
@@ -453,11 +492,20 @@ void SNamingTokensEditableTextBox::OnEditableTextChanged(const FText& InTokenize
 {
 	ValidateTokenizedText(InTokenizedText);
 
+	UpdateFilters();
+	
 	OnTokenizedTextChanged.ExecuteIfBound(InTokenizedText);
 }
 
 void SNamingTokensEditableTextBox::OnEditableTextCommitted(const FText& InTokenizedText, ETextCommit::Type InCommitType)
 {
+	const EFocusCause FocusType = DropdownAnchor->HasAnyUserFocus().Get(EFocusCause::Cleared);
+
+	if (InCommitType != ETextCommit::OnUserMovedFocus && FocusType == EFocusCause::Cleared)
+	{
+		CloseSuggestionDropdown();
+	}
+	
 	ValidateTokenizedText(InTokenizedText);
 
 	if (!TokenizedText.IsBound())
@@ -473,6 +521,19 @@ void SNamingTokensEditableTextBox::OnEditableTextCommitted(const FText& InTokeni
 	}
 }
 
+void SNamingTokensEditableTextBox::UpdateFilters() const
+{
+	if (EditableText.IsValid() && FilterWidget.IsValid() && DropdownAnchor->IsOpen())
+	{
+		// Find the text we're typing in between brackets and pass that as a filter to our suggestion drop down.
+
+		if (const TSharedPtr<FCursorTokenData> TokenLocation = GetCurrentCursorTokenData())
+		{
+			FilterWidget->FilterTreeItems(TokenLocation->EnteredWord);
+		}
+	}
+}
+
 void SNamingTokensEditableTextBox::EvaluateNamingTokens()
 {
 	check(bHandlingEvaluation);
@@ -486,6 +547,234 @@ void SNamingTokensEditableTextBox::EvaluateNamingTokens()
 	ResolvedText.Set(EvaluationData.EvaluatedText);
 }
 
+void SNamingTokensEditableTextBox::OpenOrCloseSuggestionDropdown() const
+{
+	if (DropdownAnchor.IsValid())
+	{
+		if (const TSharedPtr<FCursorTokenData> TokenLocation = GetCurrentCursorTokenData())
+		{
+			const bool bLastTokenDataChanged = !LastCursorTokenData.IsValid() || LastCursorTokenData->EnteredWord != TokenLocation->EnteredWord;
+			LastCursorTokenData = TokenLocation;
+			OpenSuggestionDropdown();
+			if (bLastTokenDataChanged)
+			{
+				UpdateFilters();
+			}
+		}
+		else
+		{
+			LastCursorTokenData.Reset();
+			CloseSuggestionDropdown();
+		}
+	}
+}
+
+void SNamingTokensEditableTextBox::OpenOrCloseSuggestionDropdownNextTick() const
+{
+	if (!bModifyingSuggestionDropdown && DropdownAnchor.IsValid())
+	{
+		bModifyingSuggestionDropdown = true;
+		
+		const_cast<SNamingTokensEditableTextBox*>(this)->RegisterActiveTimer(0.f,
+			FWidgetActiveTimerDelegate::CreateLambda([this](double InCurrentTime, float InDeltaTime)
+		{
+			bModifyingSuggestionDropdown = false;
+			OpenOrCloseSuggestionDropdown();
+			return EActiveTimerReturnType::Stop;
+		}));
+	}
+}
+
+void SNamingTokensEditableTextBox::OpenSuggestionDropdown() const
+{
+	check(DropdownAnchor.IsValid());
+	if (!DropdownAnchor->IsOpen())
+	{
+		DropdownAnchor->SetIsOpen(true, false);
+	}
+}
+
+void SNamingTokensEditableTextBox::CloseSuggestionDropdown() const
+{
+	if (DropdownAnchor.IsValid())
+	{
+		DropdownAnchor->SetIsOpen(false, false);
+	}
+}
+
+TSharedPtr<SNamingTokensEditableTextBox::FCursorTokenData> SNamingTokensEditableTextBox::GetCurrentCursorTokenData() const
+{
+	check(EditableText.IsValid());
+	
+	const FTextLocation Cursor = EditableText->GetCursorLocation();
+
+	FString Line;
+	EditableText->GetTextLine(Cursor.GetLineIndex(), Line);
+
+	const int32 CursorOffset = Cursor.GetOffset();
+	if (CursorOffset < 0 || CursorOffset >= Line.Len())
+	{
+		return nullptr;
+	}
+
+	// Look left for open '{'.
+	const FString Left = Line.Left(CursorOffset);
+
+	int32 LastOpen  = INDEX_NONE;
+	int32 LastClose = INDEX_NONE;
+	Left.FindLastChar(TEXT('{'), LastOpen);
+	Left.FindLastChar(TEXT('}'), LastClose);
+
+	// Error if no open bracket found, or a close bracket was found instead.
+	if (LastOpen == INDEX_NONE || (LastClose != INDEX_NONE && LastClose > LastOpen))
+	{
+		return nullptr;
+	}
+
+	// Look right for close '}'.
+	const FString Right   = Line.Mid(CursorOffset);
+
+	const int32 NextClose = Right.Find(TEXT("}"), ESearchCase::CaseSensitive);
+	const int32 NextOpen  = Right.Find(TEXT("{"), ESearchCase::CaseSensitive);
+
+	// Error if no close bracket found, or an open bracket was found instead.
+	if (NextClose == INDEX_NONE || (NextOpen != INDEX_NONE && NextOpen < NextClose))
+	{
+		return nullptr;
+	}
+
+	const int32 NextCloseOffset = CursorOffset + NextClose;
+
+	TSharedRef<FCursorTokenData> CursorTokenData = MakeShared<FCursorTokenData>();
+	CursorTokenData->StartLocation = FTextLocation(Cursor.GetLineIndex(), LastOpen + 1 /* Offset for token */);
+	CursorTokenData->EndLocation   = FTextLocation(Cursor.GetLineIndex(), NextCloseOffset);
+	CursorTokenData->Length        = NextCloseOffset - LastOpen;
+
+	// Record word entered between brackets.
+	{
+		FString Between = Line.Mid(LastOpen + 1, NextCloseOffset - (LastOpen + 1));
+		Between.TrimStartAndEndInline();
+
+		// Alphanumeric, underscore, numbers, and also spaces, so we can account for display names.
+		const FRegexPattern Pattern(TEXT("([A-Za-z0-9_ ]+(?::[A-Za-z0-9_ ]+)*:?)"));
+		FRegexMatcher Matcher(Pattern, Between);
+		
+		FString CurrentWord;
+		while (Matcher.FindNext())
+		{
+			CurrentWord = Matcher.GetCaptureGroup(1).TrimStartAndEnd();
+			break;
+		}
+
+		CursorTokenData->EnteredWord = MoveTemp(CurrentWord);
+	}
+	
+	return CursorTokenData;
+}
+
+TSharedRef<SWidget> SNamingTokensEditableTextBox::OnGetDropdownContent()
+{
+	if (!bEnableSuggestionDropdown)
+	{
+		return SNullWidget::NullWidget;
+	}
+	TSharedRef<SNamingTokenDataTreeViewWidget> Widget = SAssignNew(FilterWidget, SNamingTokenDataTreeViewWidget)
+		.OnItemSelected(this, &SNamingTokensEditableTextBox::OnSuggestionSelected);
+	return Widget;
+}
+
+FReply SNamingTokensEditableTextBox::OnEditableTextKeyChar(const FGeometry& MyGeometry, const FCharacterEvent& InCharacterEvent) const
+{
+	if (InCharacterEvent.GetCharacter() == TEXT('{'))
+	{
+		const FTextLocation FilterCursorLocation = EditableText->GetCursorLocation();
+		EditableText->InsertTextAtCursor(TEXT("}"));
+		EditableText->GoTo(FilterCursorLocation);
+	}
+	
+	return FReply::Unhandled();
+}
+
+FReply SNamingTokensEditableTextBox::OnEditableTextKeyDown(const FGeometry& MyGeometry, const FKeyEvent& InKeyEvent) const
+{
+	const FKey Key = InKeyEvent.GetKey();
+	
+	// We have to hook enter so we don't commit early. The user may be selecting a suggestion from the dropdown with enter.
+	if (Key == EKeys::Enter && DropdownAnchor.IsValid() && DropdownAnchor->IsOpen() && FilterWidget.IsValid())
+	{
+		OnSuggestionSelected(FilterWidget->GetSelectedItem());
+		return FReply::Handled();
+	}
+
+	if (!(Key == EKeys::Escape || Key == EKeys::Tab || Key == EKeys::Enter))
+	{
+		// If we're not submitting text then we want to always check if the suggestion box should open or close.
+		// Use a timer to ensure our updated text is submitted before we check cursor location.
+		OpenOrCloseSuggestionDropdownNextTick();
+	}
+
+	return FReply::Unhandled();
+}
+
+void SNamingTokensEditableTextBox::OnEditableTextCursorMoved(const FTextLocation& NewCursorPosition)
+{
+	const FTextSelection CurrentSelection = EditableText->GetSelection();
+    const FTextLocation& Begin = CurrentSelection.GetBeginning();
+    const FTextLocation& End = CurrentSelection.GetEnd();
+	
+	// FTextLocation has limited operator support.
+	const bool bIsSelected = (Begin.IsValid() && End.IsValid() && Begin < End) && (((Begin < NewCursorPosition) && (NewCursorPosition < End))
+		|| (NewCursorPosition == Begin || NewCursorPosition == End));
+
+	// If we're selecting text, avoid opening the suggestion window.
+	if (!bIsSelected || DropdownAnchor->IsOpen())
+	{
+		OpenOrCloseSuggestionDropdown();
+	}
+}
+
+void SNamingTokensEditableTextBox::OnSuggestionSelected(TSharedPtr<FNamingTokenDataTreeItem> SelectedItem) const
+{
+	if (SelectedItem.IsValid())
+	{
+		if (const TSharedPtr<FCursorTokenData> TokenLocation = GetCurrentCursorTokenData())
+		{
+			// Either insert just the token, or the namespace:token.
+			const bool bIsFullyQualified = !SelectedItem->bIsGlobal || bFullyQualifyGlobalTokenSuggestions;
+		
+			FString TextToInsert = SelectedItem->ToString(bIsFullyQualified);
+			if (SelectedItem->IsNamespace())
+			{
+				// Add the delimiter and continue the filter process.
+				TextToInsert += UE::NamingTokens::Utils::GetNamespaceDelimiter();
+			}
+			else
+			{
+				// We're finished if this is a token being inserted.
+				CloseSuggestionDropdown();
+			}
+			
+			// Set our cursor to the start and select everything in our range.
+			EditableText->GoTo(TokenLocation->StartLocation);
+			EditableText->SelectText(TokenLocation->StartLocation, TokenLocation->EndLocation);
+
+			// Insert all of our text.
+			EditableText->InsertTextAtCursor(FText::FromString(TextToInsert));
+			
+			if (SelectedItem->IsNamespace())
+			{
+				UpdateFilters();
+			}
+			else if (const TSharedPtr<FCursorTokenData> FinalizedTokenLocation = GetCurrentCursorTokenData())
+            {
+                // If we're inserting the final token offset the cursor to after the final bracket.
+                const FTextLocation AfterBracket(FinalizedTokenLocation->EndLocation.GetLineIndex(), FinalizedTokenLocation->EndLocation.GetOffset() + 1);
+                EditableText->GoTo(AfterBracket);
+            }
+		}
+	}
+}
+
 END_SLATE_FUNCTION_BUILD_OPTIMIZATION
 
-#undef LOCTEXT_NAMESPACE
+#undef LOCTEXT_NAMESPACE
\ No newline at end of file
diff --git a/Engine/Plugins/Developer/NamingTokens/Source/NamingTokensUI/Public/NamingTokensEditableText.h b/Engine/Plugins/Developer/NamingTokens/Source/NamingTokensUI/Public/NamingTokensEditableText.h
--- a/Engine/Plugins/Developer/NamingTokens/Source/NamingTokensUI/Public/NamingTokensEditableText.h
+++ b/Engine/Plugins/Developer/NamingTokens/Source/NamingTokensUI/Public/NamingTokensEditableText.h
@@ -91,7 +91,11 @@ protected:
 	/** Set naming token filter args to use during token evaluation. */
 	UPROPERTY(EditAnywhere, BlueprintReadWrite, Category=NamingTokens)
 	FNamingTokenFilterArgs FilterArgs;
-
+	
+	/** If the dropdown suggestion box should be enabled. */
+	UPROPERTY(EditAnywhere, AdvancedDisplay, Category=Behavior)
+	bool bEnableSuggestionDropdown = true;
+	
 	/** If this text box is configured for multiline. */
 	UPROPERTY(EditAnywhere, BlueprintReadWrite, Category=Appearance)
 	bool bIsMultiline = false;
@@ -107,7 +111,7 @@ protected:
 	/** Display the border image. */
 	UPROPERTY(EditAnywhere, AdvancedDisplay, BlueprintReadWrite, BlueprintSetter=SetDisplayBorderImage, Category=Appearance)
 	bool bDisplayBorderImage = false;
-	
+
 	/** The font to apply to the non-resolved tokens. */
 	UPROPERTY(EditAnywhere, AdvancedDisplay, BlueprintReadWrite, BlueprintSetter=SetWidgetArgumentStyle, Category=Appearance)
 	FTextBlockStyle ArgumentStyle;
diff --git a/Engine/Plugins/Developer/NamingTokens/Source/NamingTokensUI/Public/SNamingTokensEditableTextBox.h b/Engine/Plugins/Developer/NamingTokens/Source/NamingTokensUI/Public/SNamingTokensEditableTextBox.h
--- a/Engine/Plugins/Developer/NamingTokens/Source/NamingTokensUI/Public/SNamingTokensEditableTextBox.h
+++ b/Engine/Plugins/Developer/NamingTokens/Source/NamingTokensUI/Public/SNamingTokensEditableTextBox.h
@@ -10,10 +10,12 @@
 #include "Widgets/DeclarativeSyntaxSupport.h"
 #include "Widgets/Input/SMultiLineEditableTextBox.h"
 
-class FString;
-class FText;
 class ITextLayoutMarshaller;
 class FNamingTokensStringSyntaxHighlighterMarshaller;
+class FString;
+class FText;
+class SMenuAnchor;
+
 struct FNamingTokenFilterArgs;
 
 /**
@@ -32,6 +34,8 @@ public:
 		, _DisplayErrorMessage(true)
 		, _DisplayBorderImage(true)
 		, _NamingTokenFilterArgs(nullptr)
+		, _EnableSuggestionDropdown(true)
+		, _FullyQualifyGlobalTokenSuggestions(false)
 	{}
 		/** The styling of the textbox. */
 		SLATE_STYLE_ARGUMENT(FEditableTextBoxStyle, Style)
@@ -87,6 +91,12 @@ public:
 		/** Contexts to provide to evaluation. */
 		SLATE_ARGUMENT(const TArray<TObjectPtr<UObject>>*, Contexts)
 
+		/** If drop down suggestions should be enabled. */
+		SLATE_ARGUMENT(bool, EnableSuggestionDropdown)
+		
+		/** If the namespace should be inserted when selecting a global suggestion. */
+		SLATE_ARGUMENT(bool, FullyQualifyGlobalTokenSuggestions)
+
 	SLATE_END_ARGS()
 
 	/** Constructs this widget with InArgs */
@@ -113,6 +123,9 @@ public:
 	/** Sets the argument style to use. */
 	NAMINGTOKENSUI_API void SetArgumentStyle(const FTextBlockStyle* InArgumentStyle);
 
+	/** If the suggestion box should be enabled. */
+	NAMINGTOKENSUI_API void SetEnableSuggestionDropdown(bool bInEnableSuggestionDropdown);
+	
 	/** Provide filter args to use during evaluation. Passing null clears existing. */
 	NAMINGTOKENSUI_API void SetFilterArgs(const FNamingTokenFilterArgs* InFilterArgs);
 
@@ -132,6 +145,7 @@ private:
 
 	NAMINGTOKENSUI_API virtual bool HasKeyboardFocus() const override;
 	NAMINGTOKENSUI_API virtual FReply OnFocusReceived(const FGeometry& MyGeometry, const FFocusEvent& InFocusEvent) override;
+	NAMINGTOKENSUI_API virtual void OnFocusChanging(const FWeakWidgetPath& PreviousFocusPath, const FWidgetPath& NewWidgetPath, const FFocusEvent& InFocusEvent) override;
 	NAMINGTOKENSUI_API virtual FReply OnKeyDown(const FGeometry& MyGeometry, const FKeyEvent& InKeyEvent) override;
 
 	bool ValidateTokenizedText(const FText& InTokenizedText);
@@ -143,7 +157,50 @@ private:
 	void OnEditableTextChanged(const FText& InTokenizedText);
 	void OnEditableTextCommitted(const FText& InTokenizedText, ETextCommit::Type InCommitType);
 
+	/** Check for type text and update filters if needed. */
+	void UpdateFilters() const;
+	
 	void EvaluateNamingTokens();
+
+	struct FCursorTokenData
+	{
+		/** Start '{' location of token. */
+		FTextLocation StartLocation;
+		/** End '}' location of token.  */
+		FTextLocation EndLocation;
+		/** Length of EndLocation-StartLocation offset. */
+		int32 Length;
+		/** Word entered between brackets. */
+		FString EnteredWord;
+	};
+
+	/** Opens or closes the dropdown based on cursor status. */
+	void OpenOrCloseSuggestionDropdown() const;
+	/** Opens or closes the suggestion drop down on the next tick, providing a request hasn't already been made. */
+	void OpenOrCloseSuggestionDropdownNextTick() const;
+	
+	/** Trigger the suggestion process. */
+	void OpenSuggestionDropdown() const;
+	/** Close the suggestion window. */
+	void CloseSuggestionDropdown() const;
+
+	/** Retrieve information about the current token where the cursor is. */
+	TSharedPtr<FCursorTokenData> GetCurrentCursorTokenData() const;
+	
+	/** Create the dropdown content for filtering naming tokens. */
+	TSharedRef<SWidget> OnGetDropdownContent();
+	
+	/** When our editable text has a character typed. */
+	FReply OnEditableTextKeyChar(const FGeometry& MyGeometry, const FCharacterEvent& InCharacterEvent) const;
+
+	/** When editable text has a key pressed. */
+	FReply OnEditableTextKeyDown(const FGeometry& MyGeometry, const FKeyEvent& InKeyEvent) const;
+
+	/** User has moved the text cursor. */
+	void OnEditableTextCursorMoved(const FTextLocation& NewCursorPosition);
+	
+	/** When a dropdown suggestion has been selected. */
+	void OnSuggestionSelected(TSharedPtr<struct FNamingTokenDataTreeItem> SelectedItem) const;
 	
 private:
 	/** Filter args to use during evaluation. */
@@ -155,6 +212,15 @@ private:
 	/** Editable text widget. */
 	TSharedPtr<SMultiLineEditableText> EditableText;
 
+	/** The dropdown filter UI anchor. */
+	TSharedPtr<SMenuAnchor> DropdownAnchor;
+
+	/** The Dropdown widget. */
+	TSharedPtr<class SNamingTokenDataTreeViewWidget> FilterWidget;
+
+	/** Cursor token data of the last token data. */
+	mutable TSharedPtr<FCursorTokenData> LastCursorTokenData;
+	
 	/** Style shared between Editable and NonEditable text widgets. */
 	const FEditableTextBoxStyle* TextBoxStyle = nullptr;
 	/** The style to use for our argument formatting. */
@@ -206,4 +272,13 @@ private:
 
 	/** If this widget is handling all token evaluation. */
 	bool bHandlingEvaluation = false;
+
+	/** If global tokens should have their namespace inserted too from suggestions. */
+	bool bFullyQualifyGlobalTokenSuggestions = false;
+
+	/** If suggestions are enabled by default. */
+	bool bEnableSuggestionDropdown = true;
+	
+	/** Suggesting box in the process of opening or closing. */
+	mutable bool bModifyingSuggestionDropdown = false;
 };
diff --git a/Engine/Plugins/VirtualProduction/Takes/Source/TakeRecorderNamingTokens/Private/TakeRecorderNamingTokens.cpp b/Engine/Plugins/VirtualProduction/Takes/Source/TakeRecorderNamingTokens/Private/TakeRecorderNamingTokens.cpp
--- a/Engine/Plugins/VirtualProduction/Takes/Source/TakeRecorderNamingTokens/Private/TakeRecorderNamingTokens.cpp
+++ b/Engine/Plugins/VirtualProduction/Takes/Source/TakeRecorderNamingTokens/Private/TakeRecorderNamingTokens.cpp
@@ -26,6 +26,7 @@
 UTakeRecorderNamingTokens::UTakeRecorderNamingTokens()
 {
 	Namespace = ITakeRecorderNamingTokensModule::GetTakeRecorderNamespace();
+	NamespaceDisplayName = LOCTEXT("NamespaceDisplayName", "Take Recorder");
 }
 
 UTakeRecorderNamingTokens::~UTakeRecorderNamingTokens()
