NamingTokens: Add several validity checks on autocomplete handling.

diff --git a/Engine/Plugins/Developer/NamingTokens/Source/NamingTokensUI/Private/SNamingTokensEditableTextBox.cpp b/Engine/Plugins/Developer/NamingTokens/Source/NamingTokensUI/Private/SNamingTokensEditableTextBox.cpp
--- a/Engine/Plugins/Developer/NamingTokens/Source/NamingTokensUI/Private/SNamingTokensEditableTextBox.cpp
+++ b/Engine/Plugins/Developer/NamingTokens/Source/NamingTokensUI/Private/SNamingTokensEditableTextBox.cpp
@@ -364,7 +364,7 @@ FReply SNamingTokensEditableTextBox::OnKeyDown(const FGeometry& MyGeometry, cons
 		// Clear focus
 		return FReply::Handled().SetUserFocus(SharedThis(this), EFocusCause::Cleared);
 	}
-	if (DropdownAnchor->IsOpen() && FilterWidget.IsValid())
+	if (DropdownAnchor.IsValid() && DropdownAnchor->IsOpen() && FilterWidget.IsValid())
 	{
 		FilterWidget->ForwardKeyEventForNavigation(InKeyEvent);
 		// Set focus back to the editable text so the user can keep typing.
@@ -499,11 +499,14 @@ void SNamingTokensEditableTextBox::OnEditableTextChanged(const FText& InTokenize
 
 void SNamingTokensEditableTextBox::OnEditableTextCommitted(const FText& InTokenizedText, ETextCommit::Type InCommitType)
 {
-	const EFocusCause FocusType = DropdownAnchor->HasAnyUserFocus().Get(EFocusCause::Cleared);
-
-	if (InCommitType != ETextCommit::OnUserMovedFocus && FocusType == EFocusCause::Cleared)
+	if (DropdownAnchor.IsValid())
 	{
-		CloseSuggestionDropdown();
+		const EFocusCause FocusType = DropdownAnchor->HasAnyUserFocus().Get(EFocusCause::Cleared);
+
+		if (InCommitType != ETextCommit::OnUserMovedFocus && FocusType == EFocusCause::Cleared)
+		{
+			CloseSuggestionDropdown();
+		}
 	}
 	
 	ValidateTokenizedText(InTokenizedText);
@@ -727,7 +730,7 @@ void SNamingTokensEditableTextBox::OnEditableTextCursorMoved(const FTextLocation
 		|| (NewCursorPosition == Begin || NewCursorPosition == End));
 
 	// If we're selecting text, avoid opening the suggestion window.
-	if (!bIsSelected || DropdownAnchor->IsOpen())
+	if (!bIsSelected || (DropdownAnchor.IsValid() && DropdownAnchor->IsOpen()))
 	{
 		OpenOrCloseSuggestionDropdown();
 	}
