NamingTokens: Letters are now highlighted matching the search terms.
When searching for the best match priority is now given to a case sensitive search over case-insensitive.

diff --git a/Engine/Plugins/Developer/NamingTokens/Source/NamingTokensUI/Private/SNamingTokensDataTreeView.cpp b/Engine/Plugins/Developer/NamingTokens/Source/NamingTokensUI/Private/SNamingTokensDataTreeView.cpp
--- a/Engine/Plugins/Developer/NamingTokens/Source/NamingTokensUI/Private/SNamingTokensDataTreeView.cpp
+++ b/Engine/Plugins/Developer/NamingTokens/Source/NamingTokensUI/Private/SNamingTokensDataTreeView.cpp
@@ -6,12 +6,14 @@
 #include "NamingTokensEngineSubsystem.h"
 
 #include "Engine/Engine.h"
+#include "Widgets/Input/SEditableTextBox.h"
 
 #define LOCTEXT_NAMESPACE "SNamingTokenDataTreeView"
 
-void SNamingTokenDataTreeView::Construct(const FArguments& InArgs)
+void SNamingTokenDataTreeView::Construct(const FArguments& InArgs, const TSharedPtr<SNamingTokenDataTreeViewWidget>& InTreeViewWidget)
 {
 	STreeView::Construct(InArgs);
+	OwningTreeViewWidget = InTreeViewWidget;
 }
 
 void SNamingTokenDataTreeView::SetExpansionStateFromItems(const TArray<FNamingTokenDataTreeItemPtr>& InTreeItems)
@@ -23,12 +25,12 @@ void SNamingTokenDataTreeView::SetExpansionStateFromItems(const TArray<FNamingTo
 	}
 }
 
-void SNamingTokenDataTreeViewRow::Construct(const FArguments& InArgs, const TSharedPtr<SNamingTokenDataTreeView>& TreeView)
+void SNamingTokenDataTreeViewRow::Construct(const FArguments& InArgs, const TSharedRef<STableViewBase>& InOwnerTable)
 {
 	Item = InArgs._Item;
 
 	const FSuperRowType::FArguments SuperArgs;
-	SMultiColumnTableRow::Construct(SuperArgs, TreeView.ToSharedRef());
+	SMultiColumnTableRow::Construct(SuperArgs, InOwnerTable);
 }
 
 TSharedRef<SWidget> SNamingTokenDataTreeViewRow::GenerateWidgetForColumn(const FName& ColumnName)
@@ -36,54 +38,85 @@ TSharedRef<SWidget> SNamingTokenDataTreeViewRow::GenerateWidgetForColumn(const F
 	const FNamingTokenDataTreeItemPtr ItemPtr = Item.Pin();
 	if (ItemPtr.IsValid())
 	{
-		const TSharedPtr<SBorder> Border = SNew(SBorder)
-			.BorderImage(FAppStyle::Get().GetBrush("NoBorder"));
-		
-		if (ItemPtr->NamingTokenData.IsValid() && ColumnName == SNamingTokenDataTreeViewWidget::PrimaryColumnName())
-		{
-			Border->SetContent(
-				SNew(SVerticalBox)
-				+ SVerticalBox::Slot()
-				.AutoHeight()
-				.Padding(16.f, 2.f)
-				[
-					SNew(STextBlock)
-					.Text(FText::FromString(ItemPtr->NamingTokenData->TokenKey))
-				]);
-		}
-		else if (!ItemPtr->Namespace.IsEmpty())
+		if (const TSharedPtr<SNamingTokenDataTreeView> TreeViewPtr = StaticCastSharedPtr<SNamingTokenDataTreeView>(OwnerTablePtr.Pin()))
 		{
-			const FText NamespaceDisplayName = !ItemPtr->NamespaceDisplayName.IsEmpty() ?
-				FText::Format(LOCTEXT("NamespaceDisplayName", "{0} ({1}:)"),
-				ItemPtr->NamespaceDisplayName, FText::FromString(ItemPtr->Namespace))
-					: FText::FromString(ItemPtr->Namespace);
+			const TWeakPtr<SNamingTokenDataTreeViewWidget> TreeViewWidgetWeakPtr = TreeViewPtr->GetOwningTreeViewWidget();
+			if (const TSharedPtr<SNamingTokenDataTreeViewWidget> TreeViewWidgetPin = TreeViewWidgetWeakPtr.Pin())
+			{
+				const TSharedPtr<SBorder> Border = SNew(SBorder)
+					.BorderImage(FAppStyle::Get().GetBrush("NoBorder"));
+		
+				if (ItemPtr->NamingTokenData.IsValid() && ColumnName == SNamingTokenDataTreeViewWidget::PrimaryColumnName())
+				{
+					const FString FilterString = TreeViewWidgetPin->GetTokenFilterText();
+					const FText FilterText = ItemPtr->NamingTokenData->TokenKey.StartsWith(FilterString)
+						? FText::FromString(FilterString) : FText::GetEmpty();
+					
+					Border->SetContent(
+						SNew(SVerticalBox)
+						+ SVerticalBox::Slot()
+						.AutoHeight()
+						.Padding(16.f, 2.f)
+						[
+							SNew(SEditableTextBox)
+							.Style(&GetCustomEditableTextBoxStyle())
+							.Visibility(EVisibility::HitTestInvisible)
+							.IsReadOnly(true)
+							.Padding(0.f)
+							.BackgroundColor(FLinearColor::Transparent)
+							.ForegroundColor(FSlateColor::UseForeground())
+							.ReadOnlyForegroundColor(FSlateColor::UseForeground())
+							.Text(FText::FromString(ItemPtr->NamingTokenData->TokenKey))
+							.SearchText(FilterText)
+						]);
+				}
+				else if (!ItemPtr->Namespace.IsEmpty())
+				{
+					const FText NamespaceDisplayName = !ItemPtr->NamespaceDisplayName.IsEmpty() ?
+						FText::Format(LOCTEXT("NamespaceDisplayName", "{0} ({1}:)"),
+						ItemPtr->NamespaceDisplayName, FText::FromString(ItemPtr->Namespace))
+							: FText::FromString(ItemPtr->Namespace);
 			
-			Border->SetContent(SNew(SHorizontalBox)
-				+ SHorizontalBox::Slot()
-				.AutoWidth()
-				.Padding(0.f, 0.f, 0.f, 0.f)
-				.VAlign(VAlign_Center)
-				[
-					SNew(SExpanderArrow, SharedThis(this))
-					.StyleSet(&FAppStyle::Get())
-					.IndentAmount(5)
-				]
-				+ SHorizontalBox::Slot()
-				.FillWidth(1.0f)
-				.VAlign(VAlign_Center)
-				[
-					SNew(STextBlock)
-					.Font(FCoreStyle::Get().GetFontStyle("ExpandableArea.TitleFont"))
-					.Text(NamespaceDisplayName)
-				]);
-		}
+					Border->SetContent(SNew(SHorizontalBox)
+						+ SHorizontalBox::Slot()
+						.AutoWidth()
+						.Padding(0.f, 0.f, 0.f, 0.f)
+						.VAlign(VAlign_Center)
+						[
+							SNew(SExpanderArrow, SharedThis(this))
+							.StyleSet(&FAppStyle::Get())
+							.IndentAmount(5)
+						]
+						+ SHorizontalBox::Slot()
+						.FillWidth(1.0f)
+						.VAlign(VAlign_Center)
+						[
+							SNew(STextBlock)
+							.Font(FCoreStyle::Get().GetFontStyle("ExpandableArea.TitleFont"))
+							.Text(NamespaceDisplayName)
+						]);
+				}
 
-		return Border.ToSharedRef();
+				return Border.ToSharedRef();
+			}
+		}
 	}
 
 	return SNullWidget::NullWidget;
 }
 
+const FEditableTextBoxStyle& SNamingTokenDataTreeViewRow::GetCustomEditableTextBoxStyle()
+{
+	static FEditableTextBoxStyle EditableTextBoxStyle = []()
+		{
+			static FEditableTextBoxStyle Style = FCoreStyle::Get().GetWidgetStyle<FEditableTextBoxStyle>("NormalEditableTextBox");
+			Style.BackgroundImageReadOnly = *FCoreStyle::Get().GetBrush("NoBorder");
+
+			return Style;
+		}();
+	return EditableTextBoxStyle;
+}
+
 SNamingTokenDataTreeViewWidget::~SNamingTokenDataTreeViewWidget()
 {
 }
@@ -115,7 +148,7 @@ void SNamingTokenDataTreeViewWidget::Construct(const FArguments& InArgs)
 		.MaxDesiredHeight(250.f)
 		.Padding(1.f)
 		[
-			SAssignNew(TreeView, SNamingTokenDataTreeView)
+			SAssignNew(TreeView, SNamingTokenDataTreeView, SharedThis(this))
 			.SelectionMode(ESelectionMode::Single)
 			.TreeItemsSource(&FilteredRootTreeItems)
 			.HeaderRow(HeaderRow)
@@ -196,6 +229,9 @@ void SNamingTokenDataTreeViewWidget::FilterTreeItems(const FString& InFilterText
 	FString CompleteOrPartialNamespace = UE::NamingTokens::Utils::GetNamespaceFromTokenKey(InFilterText);
 	FString PartialToken = UE::NamingTokens::Utils::RemoveNamespaceFromTokenKey(InFilterText);
 
+	TokenFilterText = PartialToken;
+	NamespaceFilterText = CompleteOrPartialNamespace;
+	
 	const bool bSearchNamespaceAndToken = CompleteOrPartialNamespace.IsEmpty();
 	if (bSearchNamespaceAndToken)
 	{
@@ -384,7 +420,7 @@ void SNamingTokenDataTreeViewWidget::FinalizeSelection(FNamingTokenDataTreeItemP
 TSharedRef<ITableRow> SNamingTokenDataTreeViewWidget::OnGenerateRowForTree(FNamingTokenDataTreeItemPtr Item,
                                                                            const TSharedRef<STableViewBase>& OwnerTable)
 {
-	return SNew(SNamingTokenDataTreeViewRow, TreeView).Item(Item);
+	return SNew(SNamingTokenDataTreeViewRow, OwnerTable).Item(Item);
 }
 
 void SNamingTokenDataTreeViewWidget::OnGetChildrenForTree(FNamingTokenDataTreeItemPtr InParent, TArray<FNamingTokenDataTreeItemPtr>& OutChildren)
diff --git a/Engine/Plugins/Developer/NamingTokens/Source/NamingTokensUI/Private/SNamingTokensDataTreeView.h b/Engine/Plugins/Developer/NamingTokens/Source/NamingTokensUI/Private/SNamingTokensDataTreeView.h
--- a/Engine/Plugins/Developer/NamingTokens/Source/NamingTokensUI/Private/SNamingTokensDataTreeView.h
+++ b/Engine/Plugins/Developer/NamingTokens/Source/NamingTokensUI/Private/SNamingTokensDataTreeView.h
@@ -60,11 +60,16 @@ struct FNamingTokenDataTreeItem : TSharedFromThis<FNamingTokenDataTreeItem>
 
 		check(NamingTokenData.IsValid());
 
-		if (NamingTokenData->TokenKey == FilterText)
+		if (NamingTokenData->TokenKey.Equals(FilterText, ESearchCase::CaseSensitive))
 		{
 			return 1.f;
 		}
 		
+		if (NamingTokenData->TokenKey.Equals(FilterText, ESearchCase::IgnoreCase))
+		{
+			return 0.9f;
+		}
+		
 		if (NamingTokenData->TokenKey.StartsWith(FilterText, ESearchCase::IgnoreCase)
 			|| NamingTokenData->DisplayName.ToString().StartsWith(FilterText, ESearchCase::IgnoreCase))
 		{
@@ -106,10 +111,17 @@ typedef TSharedPtr<FNamingTokenDataTreeItem> FNamingTokenDataTreeItemPtr;
 class SNamingTokenDataTreeView : public STreeView<FNamingTokenDataTreeItemPtr>
 {
 public:
-	void Construct(const FArguments& InArgs);
+	void Construct(const FArguments& InArgs, const TSharedPtr<class SNamingTokenDataTreeViewWidget>& InTreeViewWidget);
 
 	/** Recursively set expansion state of tree view to match items. */
 	void SetExpansionStateFromItems(const TArray<FNamingTokenDataTreeItemPtr>& InTreeItems);
+
+	/** Pointer to our owning widget. */
+	const TWeakPtr<SNamingTokenDataTreeViewWidget>& GetOwningTreeViewWidget() const { return OwningTreeViewWidget; }
+	
+private:
+	/** The owning tree view widget. */
+	TWeakPtr<SNamingTokenDataTreeViewWidget> OwningTreeViewWidget;
 };
 
 /**
@@ -124,12 +136,15 @@ public:
 	SLATE_ARGUMENT(FNamingTokenDataTreeItemPtr, Item)
 	SLATE_END_ARGS()
 	
-	void Construct(const FArguments& InArgs, const TSharedPtr<SNamingTokenDataTreeView>& TreeView);
+	void Construct(const FArguments& InArgs, const TSharedRef<STableViewBase>& InOwnerTable);
 
 	// ~Begin SMultiColumnTableRow interface
 	virtual TSharedRef<SWidget> GenerateWidgetForColumn(const FName& ColumnName) override;
 	// ~End SMultiColumnTableRow interface
 
+private:
+	static const FEditableTextBoxStyle& GetCustomEditableTextBoxStyle();
+
 private:
 	/** The item associated with this row of data. */
 	TWeakPtr<FNamingTokenDataTreeItem> Item;
@@ -168,6 +183,13 @@ public:
 
 	/** Retrieve an item given its absolute index. */
 	FNamingTokenDataTreeItemPtr GetItemFromIndex(int32 InIndex) const;
+
+	/** Get the token filter text applied. */
+	const FString& GetTokenFilterText() const { return TokenFilterText; }
+
+	/** Get the namespace filter text applied. */
+	const FString& GetNamespaceFilterText() const { return NamespaceFilterText; }
+	
 private:
 	FNamingTokenDataTreeItemPtr GetItemFromIndex(int32 InIndex, const TArray<FNamingTokenDataTreeItemPtr>& InItems) const;
 
@@ -208,6 +230,12 @@ private:
 
 	/** The last filter applied. */
 	FString LastFilterText;
+
+	/** Last token filter text applied. */
+	FString TokenFilterText;
+
+	/** Last namespace filter text applied. */
+	FString NamespaceFilterText;
 	
 	/** Index of tree item currently selected. */
 	int32 CurrentSelectionIndex = 0;
