NamingTokens: Make displaying resolved text in the UMG widget configurable.
Add BlueprintGetters for several functions.

diff --git a/Engine/Plugins/Developer/NamingTokens/Source/NamingTokensUI/Private/NamingTokensEditableText.cpp b/Engine/Plugins/Developer/NamingTokens/Source/NamingTokensUI/Private/NamingTokensEditableText.cpp
--- a/Engine/Plugins/Developer/NamingTokens/Source/NamingTokensUI/Private/NamingTokensEditableText.cpp
+++ b/Engine/Plugins/Developer/NamingTokens/Source/NamingTokensUI/Private/NamingTokensEditableText.cpp
@@ -77,6 +77,16 @@ bool UNamingTokensEditableText::GetDisplayBorderImage() const
 	return bDisplayBorderImage;
 }
 
+void UNamingTokensEditableText::SetCanDisplayResolvedText(bool bValue)
+{
+	bCanDisplayResolvedText = bValue;
+}
+
+bool UNamingTokensEditableText::GetCanDisplayResolvedText() const
+{
+	return bCanDisplayResolvedText;
+}
+
 #if WITH_EDITOR
 void UNamingTokensEditableText::PostEditChangeProperty(FPropertyChangedEvent& PropertyChangedEvent)
 {
@@ -113,9 +123,7 @@ TSharedRef<SWidget> UNamingTokensEditableText::RebuildWidget()
 	MyNamingTokensEditableText = SNew(SNamingTokensEditableTextBox)
 		.Text(GetText())
 		.IsReadOnly(GetIsReadOnly())
-#if WITH_EDITOR
-		.CanDisplayResolvedText(MakeAttributeUObject(this, &UNamingTokensEditableText::CanDisplayResolvedTextInEditor))
-#endif
+		.CanDisplayResolvedText(MakeAttributeUObject(this, &UNamingTokensEditableText::CanDisplayResolvedTextAttribute))
 		.AllowMultiLine(bIsMultiline)
 		.Style(&*EditableTextBoxStyle)
 		.ArgumentStyle(&ArgumentStyle)
@@ -170,20 +178,23 @@ void UNamingTokensEditableText::ReleaseSlateResources(bool bReleaseChildren)
 	MyNamingTokensEditableText.Reset();
 }
 
-#if WITH_EDITOR
-
-bool UNamingTokensEditableText::CanDisplayResolvedTextInEditor() const
+bool UNamingTokensEditableText::CanDisplayResolvedTextAttribute() const
 {
-	if (IsDesignTime())
+	if (GetCanDisplayResolvedText())
 	{
-		return bDisplayResolvedTextInDesigner;
+#if WITH_EDITOR
+		if (IsDesignTime())
+		{
+			return bDisplayResolvedTextInDesigner;
+		}
+#endif
+	
+		return true;
 	}
-
-	return true;
+	
+	return false;
 }
 
-#endif
-
 void UNamingTokensEditableText::OnPreEvaluateNamingTokensCallback()
 {
 	OnPreEvaluateNamingTokens.Broadcast();
diff --git a/Engine/Plugins/Developer/NamingTokens/Source/NamingTokensUI/Private/SNamingTokensEditableTextBox.cpp b/Engine/Plugins/Developer/NamingTokens/Source/NamingTokensUI/Private/SNamingTokensEditableTextBox.cpp
--- a/Engine/Plugins/Developer/NamingTokens/Source/NamingTokensUI/Private/SNamingTokensEditableTextBox.cpp
+++ b/Engine/Plugins/Developer/NamingTokens/Source/NamingTokensUI/Private/SNamingTokensEditableTextBox.cpp
@@ -269,7 +269,7 @@ FText SNamingTokensEditableTextBox::GetText() const
 	}
 	else
 	{
-		if (ResolvedText.IsSet() && CanDisplayResolvedText.Get())
+		if (ResolvedText.IsSet() && CanDisplayResolvedText.Get(true))
 		{
 			return ResolvedText.Get();
 		}
diff --git a/Engine/Plugins/Developer/NamingTokens/Source/NamingTokensUI/Public/NamingTokensEditableText.h b/Engine/Plugins/Developer/NamingTokens/Source/NamingTokensUI/Public/NamingTokensEditableText.h
--- a/Engine/Plugins/Developer/NamingTokens/Source/NamingTokensUI/Public/NamingTokensEditableText.h
+++ b/Engine/Plugins/Developer/NamingTokens/Source/NamingTokensUI/Public/NamingTokensEditableText.h
@@ -56,13 +56,24 @@ public:
 	NAMINGTOKENSUI_API void SetDisplayBorderImage(bool bValue);
 
 	/** If the token icon should be displayed. */
+	UFUNCTION(BlueprintGetter)
 	NAMINGTOKENSUI_API bool GetDisplayTokenIcon() const;
 
 	/** If the error icon should be displayed. */
+	UFUNCTION(BlueprintGetter)
 	NAMINGTOKENSUI_API bool GetDisplayErrorMessage() const;
 
 	/** If the border image should be displayed. */
+	UFUNCTION(BlueprintGetter)
 	NAMINGTOKENSUI_API bool GetDisplayBorderImage() const;
+
+	/** Set whether we can display resolved text or not. */
+	UFUNCTION(BlueprintSetter)
+	NAMINGTOKENSUI_API void SetCanDisplayResolvedText(bool bValue);
+
+	/** Get if we can display resolved text. */
+	UFUNCTION(BlueprintGetter)
+	NAMINGTOKENSUI_API bool GetCanDisplayResolvedText() const;
 	
 #if WITH_EDITOR
 	virtual void PostEditChangeProperty(FPropertyChangedEvent& PropertyChangedEvent) override;
@@ -79,10 +90,8 @@ protected:
 	// End UVisual Interface
 	
 private:
-#if WITH_EDITOR
-	/** If we're allowed to display the resolved text. */
-	bool CanDisplayResolvedTextInEditor() const;
-#endif
+	/** Calculate if we're allowed to display the resolved text. */
+	bool CanDisplayResolvedTextAttribute() const;
 
 	/** Callback when we are evaluating tokens. */
 	void OnPreEvaluateNamingTokensCallback();
@@ -101,15 +110,15 @@ protected:
 	bool bIsMultiline = false;
 
 	/** Display the token icon in the text box. */
-	UPROPERTY(EditAnywhere, AdvancedDisplay, BlueprintReadWrite, BlueprintSetter=SetDisplayTokenIcon, Category=Appearance)
+	UPROPERTY(EditAnywhere, AdvancedDisplay, BlueprintReadWrite, BlueprintSetter=SetDisplayTokenIcon, BlueprintGetter=GetDisplayTokenIcon, Category=Appearance)
 	bool bDisplayTokenIcon = false;
 
 	/** Display an error message when tokens aren't properly formatted. */
-	UPROPERTY(EditAnywhere, AdvancedDisplay, BlueprintReadWrite, BlueprintSetter=SetDisplayErrorMessage, Category=Appearance)
+	UPROPERTY(EditAnywhere, AdvancedDisplay, BlueprintReadWrite, BlueprintSetter=SetDisplayErrorMessage, BlueprintGetter=GetDisplayErrorMessage, Category=Appearance)
 	bool bDisplayErrorMessage = false;
 
 	/** Display the border image. */
-	UPROPERTY(EditAnywhere, AdvancedDisplay, BlueprintReadWrite, BlueprintSetter=SetDisplayBorderImage, Category=Appearance)
+	UPROPERTY(EditAnywhere, AdvancedDisplay, BlueprintReadWrite, BlueprintSetter=SetDisplayBorderImage, BlueprintGetter=GetDisplayBorderImage, Category=Appearance)
 	bool bDisplayBorderImage = false;
 
 	/** The font to apply to the non-resolved tokens. */
@@ -119,10 +128,14 @@ protected:
 	/** Set the background color of the border image. */
 	UPROPERTY(EditAnywhere, AdvancedDisplay, BlueprintReadWrite, BlueprintSetter=SetBackgroundColor, Category=Appearance)
 	FSlateColor BackgroundColor;
+
+	/** If the resolved text can be displayed. */
+	UPROPERTY(EditAnywhere, AdvancedDisplay, BlueprintReadWrite, BlueprintSetter=SetCanDisplayResolvedText, BlueprintGetter=GetCanDisplayResolvedText, Category=Appearance)
+	bool bCanDisplayResolvedText = true;
 	
 #if WITH_EDITORONLY_DATA
 	/** If the resolved text should display while designing in the UMG window. If false, only the tokenized text will be displayed in the UMG designer. */
-	UPROPERTY(EditAnywhere, AdvancedDisplay, Category=Appearance)
+	UPROPERTY(EditAnywhere, AdvancedDisplay, Category=Appearance, meta = (EditCondition="bCanDisplayResolvedText", EditConditionHides))
 	bool bDisplayResolvedTextInDesigner = true;
 #endif
 	
