NamingTokens: Add required function export to SNamingTokensEditableTextBox
#rnx

diff --git a/Engine/Plugins/Developer/NamingTokens/Source/NamingTokensUI/Public/SNamingTokensEditableTextBox.h b/Engine/Plugins/Developer/NamingTokens/Source/NamingTokensUI/Public/SNamingTokensEditableTextBox.h
--- a/Engine/Plugins/Developer/NamingTokens/Source/NamingTokensUI/Public/SNamingTokensEditableTextBox.h
+++ b/Engine/Plugins/Developer/NamingTokens/Source/NamingTokensUI/Public/SNamingTokensEditableTextBox.h
@@ -138,7 +138,7 @@ private:
 
 	const FSlateBrush* GetBorderImage() const;
 	/** The background color for the border. */
-	virtual FSlateColor GetBorderBackgroundColor() const;
+	NAMINGTOKENSUI_API virtual FSlateColor GetBorderBackgroundColor() const;
 	NAMINGTOKENSUI_API virtual FSlateColor GetForegroundColor() const override;
 
 	void SetError(const FText& InError) const;
