NamingTokens: Provide world context by default.

diff --git a/Engine/Plugins/Developer/NamingTokens/Source/NamingTokens/NamingTokens.Build.cs b/Engine/Plugins/Developer/NamingTokens/Source/NamingTokens/NamingTokens.Build.cs
--- a/Engine/Plugins/Developer/NamingTokens/Source/NamingTokens/NamingTokens.Build.cs
+++ b/Engine/Plugins/Developer/NamingTokens/Source/NamingTokens/NamingTokens.Build.cs
@@ -22,5 +22,15 @@ public class NamingTokens : ModuleRules
 				"CoreUObject"
 			}
 		);
+
+		if (Target.Type == TargetType.Editor)
+		{
+			PrivateDependencyModuleNames.AddRange(
+				new string[]
+				{
+					"UnrealEd"
+				}
+			);
+		}
 	}
 }
diff --git a/Engine/Plugins/Developer/NamingTokens/Source/NamingTokens/Private/NamingTokens.cpp b/Engine/Plugins/Developer/NamingTokens/Source/NamingTokens/Private/NamingTokens.cpp
--- a/Engine/Plugins/Developer/NamingTokens/Source/NamingTokens/Private/NamingTokens.cpp
+++ b/Engine/Plugins/Developer/NamingTokens/Source/NamingTokens/Private/NamingTokens.cpp
@@ -11,6 +11,10 @@
 #include "Internationalization/Regex.h"
 #include "UObject/Package.h"
 
+#if WITH_EDITOR
+#include "Editor.h"
+#endif
+
 #include UE_INLINE_GENERATED_CPP_BY_NAME(NamingTokens)
 
 #define LOCTEXT_NAMESPACE "NamingTokens"
@@ -25,6 +29,36 @@ void UNamingTokens::PostInitProperties()
 	CreateDefaultTokens();
 }
 
+UWorld* UNamingTokens::GetWorld() const
+{
+#if WITH_EDITOR
+	if (GEditor)
+	{
+		if (const FWorldContext* PIEWorldContext = GEditor->GetPIEWorldContext())
+		{
+			return PIEWorldContext->World();
+		}
+
+		return GEditor->GetEditorWorldContext(false).World();
+	}
+
+	if (IsValid(GWorld))
+	{
+		return GWorld;
+	}
+#endif
+
+	for (const FWorldContext& WorldContext : GEngine->GetWorldContexts())
+	{
+		if (WorldContext.WorldType == EWorldType::Game)
+		{
+			return WorldContext.World();
+		}
+	}
+
+	return nullptr;
+}
+
 void UNamingTokens::Validate() const
 {
 	// Namespace validation
diff --git a/Engine/Plugins/Developer/NamingTokens/Source/NamingTokens/Public/NamingTokens.h b/Engine/Plugins/Developer/NamingTokens/Source/NamingTokens/Public/NamingTokens.h
--- a/Engine/Plugins/Developer/NamingTokens/Source/NamingTokens/Public/NamingTokens.h
+++ b/Engine/Plugins/Developer/NamingTokens/Source/NamingTokens/Public/NamingTokens.h
@@ -22,6 +22,7 @@ public:
 	UE_API UNamingTokens();
 
 	UE_API virtual void PostInitProperties() override;
+	UE_API virtual UWorld* GetWorld() const override;
 
 	/** Validate internal values including namespace and token keys. */
 	UE_API void Validate() const;
