NamingTokens: Fix double closing brackets and multiline navigation.
- Typing in a closing bracket will check if the next space already is one and offset the cursor instead of adding another bracket.
- Fix nagivation events not being recognized if the textbox is configured as multiline.

diff --git a/Engine/Plugins/Developer/NamingTokens/Source/NamingTokensUI/Private/SNamingTokensDataTreeView.cpp b/Engine/Plugins/Developer/NamingTokens/Source/NamingTokensUI/Private/SNamingTokensDataTreeView.cpp
--- a/Engine/Plugins/Developer/NamingTokens/Source/NamingTokensUI/Private/SNamingTokensDataTreeView.cpp
+++ b/Engine/Plugins/Developer/NamingTokens/Source/NamingTokensUI/Private/SNamingTokensDataTreeView.cpp
@@ -331,34 +331,33 @@ void SNamingTokenDataTreeViewWidget::FilterTreeItems(const FString& InFilterText
 	}
 }
 
-void SNamingTokenDataTreeViewWidget::ForwardKeyEventForNavigation(const FKeyEvent& InKeyEvent)
+bool SNamingTokenDataTreeViewWidget::ForwardKeyEventForNavigation(const FKeyEvent& InKeyEvent)
 {
 	check(TreeView.IsValid());
 
-	if (InKeyEvent.GetKey() == EKeys::Tab || InKeyEvent.GetKey() == EKeys::Enter)
+	const FKey Key = InKeyEvent.GetKey();
+	if (Key == EKeys::Tab || Key == EKeys::Enter)
 	{
 		FinalizeSelection(GetSelectedItem());
-		return;
+		return true;
 	}
 	
 	const EUINavigation Direction = FSlateApplicationBase::Get().GetNavigationDirectionFromKey(InKeyEvent);
-	if (Direction != EUINavigation::Invalid)
+	switch (Direction)
 	{
-		int32 NewIndex = CurrentSelectionIndex;
-		if (Direction == EUINavigation::Down)
-		{
-			NewIndex++;
-		}
-		else if (Direction == EUINavigation::Up)
-		{
+		case EUINavigation::Down:
+			SetItemFromIndex(CurrentSelectionIndex + 1);
+			return true;
+		case EUINavigation::Up:
 			if (CurrentSelectionIndex > 0)
 			{
-				NewIndex--;
+				SetItemFromIndex(CurrentSelectionIndex - 1);
 			}
-		}
-
-		SetItemFromIndex(NewIndex);
+			return true;
+		default:
+			break;
 	}
+	return false;
 }
 
 FNamingTokenDataTreeItemPtr SNamingTokenDataTreeViewWidget::GetSelectedItem() const
diff --git a/Engine/Plugins/Developer/NamingTokens/Source/NamingTokensUI/Private/SNamingTokensDataTreeView.h b/Engine/Plugins/Developer/NamingTokens/Source/NamingTokensUI/Private/SNamingTokensDataTreeView.h
--- a/Engine/Plugins/Developer/NamingTokens/Source/NamingTokensUI/Private/SNamingTokensDataTreeView.h
+++ b/Engine/Plugins/Developer/NamingTokens/Source/NamingTokensUI/Private/SNamingTokensDataTreeView.h
@@ -175,8 +175,11 @@ public:
 	/** Filter down to specific tokens and namespaces. */
 	void FilterTreeItems(const FString& InFilterText);
 	
-	/** Provide a key event from the user to determine a navigation path, highlighting an item. */
-	void ForwardKeyEventForNavigation(const FKeyEvent& InKeyEvent);
+	/**
+	 * Provide a key event from the user to determine a navigation path, highlighting an item.
+	 * @return True if the key event contained a recognized navigation key.
+	 */
+	bool ForwardKeyEventForNavigation(const FKeyEvent& InKeyEvent);
 
 	/** Retrieve the currently selected item. */
 	FNamingTokenDataTreeItemPtr GetSelectedItem() const;
diff --git a/Engine/Plugins/Developer/NamingTokens/Source/NamingTokensUI/Private/SNamingTokensEditableTextBox.cpp b/Engine/Plugins/Developer/NamingTokens/Source/NamingTokensUI/Private/SNamingTokensEditableTextBox.cpp
--- a/Engine/Plugins/Developer/NamingTokens/Source/NamingTokensUI/Private/SNamingTokensEditableTextBox.cpp
+++ b/Engine/Plugins/Developer/NamingTokens/Source/NamingTokensUI/Private/SNamingTokensEditableTextBox.cpp
@@ -358,19 +358,14 @@ void SNamingTokensEditableTextBox::OnFocusChanging(const FWeakWidgetPath& Previo
 
 FReply SNamingTokensEditableTextBox::OnKeyDown(const FGeometry& MyGeometry, const FKeyEvent& InKeyEvent)
 {
-	const FKey Key = InKeyEvent.GetKey();
-	if (Key == EKeys::Escape && EditableText->HasKeyboardFocus())
+	if (InKeyEvent.GetKey() == EKeys::Escape && EditableText->HasKeyboardFocus())
 	{
 		// Clear focus
 		return FReply::Handled().SetUserFocus(SharedThis(this), EFocusCause::Cleared);
 	}
-	if (DropdownAnchor.IsValid() && DropdownAnchor->IsOpen() && FilterWidget.IsValid())
-	{
-		FilterWidget->ForwardKeyEventForNavigation(InKeyEvent);
-		// Set focus back to the editable text so the user can keep typing.
-		return FReply::Handled().SetUserFocus(EditableText.ToSharedRef(), EFocusCause::OtherWidgetLostFocus);
-	}
-
+	// Other keydown events for navigation need to be handled in the SMultiLineEditableText callback OnEditableTextKeyDown,
+	// because this callback will not always fire if bMultiline is true.
+	
 	return FReply::Unhandled();
 }
 
@@ -690,10 +685,31 @@ FReply SNamingTokensEditableTextBox::OnEditableTextKeyChar(const FGeometry& MyGe
 {
 	if (InCharacterEvent.GetCharacter() == TEXT('{'))
 	{
+		// Create a closing bracket.
 		const FTextLocation FilterCursorLocation = EditableText->GetCursorLocation();
 		EditableText->InsertTextAtCursor(TEXT("}"));
 		EditableText->GoTo(FilterCursorLocation);
 	}
+	else if (InCharacterEvent.GetCharacter() == TEXT('}'))
+	{
+		// If we are manually typing the closing bracket, and we already have one, simply move the cursor
+		// to after the pre-existing closing bracket instead.
+		const FTextLocation CursorLocation = EditableText->GetCursorLocation();
+		
+		FString Line;
+		EditableText->GetTextLine(CursorLocation.GetLineIndex(), Line);
+		// Make sure we're not already at the end of the text line.
+		if (CursorLocation.GetOffset() < Line.Len())
+		{
+			const TCHAR NextCharacter = EditableText->GetCharacterAt(CursorLocation);
+			if (NextCharacter == TEXT('}'))
+			{
+				const FTextLocation NextCursorLocation(CursorLocation.GetLineIndex(), CursorLocation.GetOffset() + 1);
+				EditableText->GoTo(NextCursorLocation);
+				return FReply::Handled();
+			}
+		}
+	}
 	
 	return FReply::Unhandled();
 }
@@ -715,7 +731,14 @@ FReply SNamingTokensEditableTextBox::OnEditableTextKeyDown(const FGeometry& MyGe
 		// Use a timer to ensure our updated text is submitted before we check cursor location.
 		OpenOrCloseSuggestionDropdownNextTick();
 	}
-
+	
+	if (DropdownAnchor.IsValid() && DropdownAnchor->IsOpen() && FilterWidget.IsValid()
+		&& FilterWidget->ForwardKeyEventForNavigation(InKeyEvent))
+	{
+		// Set focus back to the editable text so the user can keep typing.
+		return FReply::Handled().SetUserFocus(EditableText.ToSharedRef(), EFocusCause::OtherWidgetLostFocus);
+	}
+	
 	return FReply::Unhandled();
 }
 
