NamingTokens: Fix auto complete focus issues.

- Autocomplete suggestions no longer reset if the suggestion dropdown loses focus. Clicking on the text again will filter the suggestion box appropriately.
- Double clicking on a suggestion properly commits the text and refocuses the textbox, fixing an issue where text focus is lost, and clicking on the text box will revert to the previous text before the suggestion.
- Clicking on the suggestion dropdown will maintain keyboard focus on the editable text, allowing keyboard navigation and selection to continue to work alongside mouse navigation.

diff --git a/Engine/Plugins/Developer/NamingTokens/Source/NamingTokensUI/Private/SNamingTokensDataTreeView.cpp b/Engine/Plugins/Developer/NamingTokens/Source/NamingTokensUI/Private/SNamingTokensDataTreeView.cpp
--- a/Engine/Plugins/Developer/NamingTokens/Source/NamingTokensUI/Private/SNamingTokensDataTreeView.cpp
+++ b/Engine/Plugins/Developer/NamingTokens/Source/NamingTokensUI/Private/SNamingTokensDataTreeView.cpp
@@ -25,6 +25,13 @@ void SNamingTokenDataTreeView::SetExpansionStateFromItems(const TArray<FNamingTo
 	}
 }
 
+FReply SNamingTokenDataTreeView::OnFocusReceived(const FGeometry& MyGeometry, const FFocusEvent& InFocusEvent)
+{
+	check(OwningTreeViewWidget.IsValid());
+	OwningTreeViewWidget.Pin()->NotifyFocusReceived();
+	return STreeView<FNamingTokenDataTreeItemPtr>::OnFocusReceived(MyGeometry, InFocusEvent);
+}
+
 void SNamingTokenDataTreeViewRow::Construct(const FArguments& InArgs, const TSharedRef<STableViewBase>& InOwnerTable)
 {
 	Item = InArgs._Item;
@@ -125,6 +132,7 @@ void SNamingTokenDataTreeViewWidget::Construct(const FArguments& InArgs)
 {
 	ItemSelectionChangedDelegate = InArgs._OnSelectionChanged;
 	ItemSelectedDelegate = InArgs._OnItemSelected;
+	OnFocusedDelegate = InArgs._OnFocused;
 	
 	TSharedPtr<SHeaderRow> HeaderRow;
 	// Setup the columns.
@@ -205,8 +213,11 @@ void SNamingTokenDataTreeViewWidget::PopulateTreeItems()
 						NamespaceTreeItem->ChildItems.Add(TokenTreeItem);
 					}
 				}
-				
-				RootTreeItems.Add(NamespaceTreeItem);
+
+				if (NamespaceTreeItem->ChildItems.Num() > 0)
+				{
+					RootTreeItems.Add(NamespaceTreeItem);
+				}
 			}
 		}
 	}
@@ -370,6 +381,17 @@ FNamingTokenDataTreeItemPtr SNamingTokenDataTreeViewWidget::GetItemFromIndex(int
 	return GetItemFromIndex(InIndex, FilteredRootTreeItems);
 }
 
+void SNamingTokenDataTreeViewWidget::NotifyFocusReceived()
+{
+	OnFocusedDelegate.ExecuteIfBound();
+}
+
+FReply SNamingTokenDataTreeViewWidget::OnFocusReceived(const FGeometry& MyGeometry, const FFocusEvent& InFocusEvent)
+{
+	NotifyFocusReceived();
+	return SCompoundWidget::OnFocusReceived(MyGeometry, InFocusEvent);
+}
+
 FNamingTokenDataTreeItemPtr SNamingTokenDataTreeViewWidget::GetItemFromIndex(int32 InIndex, const TArray<FNamingTokenDataTreeItemPtr>& InItems) const
 {
 	for (int32 Idx = 0; Idx < InItems.Num(); ++Idx)
@@ -397,10 +419,7 @@ void SNamingTokenDataTreeViewWidget::SetItemFromIndex(int32 InIndex)
 {
 	if (const FNamingTokenDataTreeItemPtr Item = GetItemFromIndex(InIndex))
 	{
-		if (const FNamingTokenDataTreeItemPtr CurrentSelection = GetSelectedItem())
-		{
-			TreeView->SetItemSelection(CurrentSelection, false);
-		}
+		ClearCurrentSelection();
 		
 		CurrentSelectionIndex = InIndex;
 		TreeView->SetItemSelection(Item, true);
@@ -411,6 +430,15 @@ void SNamingTokenDataTreeViewWidget::SetItemFromIndex(int32 InIndex)
 	}
 }
 
+void SNamingTokenDataTreeViewWidget::ClearCurrentSelection()
+{
+	if (const FNamingTokenDataTreeItemPtr CurrentSelection = GetSelectedItem())
+	{
+		TreeView->SetItemSelection(CurrentSelection, false);
+		CurrentSelectionIndex = INDEX_NONE;
+	}
+}
+
 void SNamingTokenDataTreeViewWidget::FinalizeSelection(FNamingTokenDataTreeItemPtr SelectedItem)
 {
 	ItemSelectedDelegate.ExecuteIfBound(SelectedItem);
@@ -434,6 +462,11 @@ void SNamingTokenDataTreeViewWidget::OnItemExpansionChanged(FNamingTokenDataTree
 
 void SNamingTokenDataTreeViewWidget::OnTreeViewItemSelectionChanged(FNamingTokenDataTreeItemPtr SelectedItem, ESelectInfo::Type SelectInfo)
 {
+	if (SelectInfo == ESelectInfo::OnMouseClick && SelectedItem.IsValid())
+	{
+		ClearCurrentSelection();
+		CurrentSelectionIndex = SelectedItem->ItemIndex;
+	}
 	ItemSelectionChangedDelegate.ExecuteIfBound(SelectedItem, SelectInfo);
 }
 
diff --git a/Engine/Plugins/Developer/NamingTokens/Source/NamingTokensUI/Private/SNamingTokensDataTreeView.h b/Engine/Plugins/Developer/NamingTokens/Source/NamingTokensUI/Private/SNamingTokensDataTreeView.h
--- a/Engine/Plugins/Developer/NamingTokens/Source/NamingTokensUI/Private/SNamingTokensDataTreeView.h
+++ b/Engine/Plugins/Developer/NamingTokens/Source/NamingTokensUI/Private/SNamingTokensDataTreeView.h
@@ -118,7 +118,11 @@ public:
 
 	/** Pointer to our owning widget. */
 	const TWeakPtr<SNamingTokenDataTreeViewWidget>& GetOwningTreeViewWidget() const { return OwningTreeViewWidget; }
-	
+
+	// ~Begin STreeView
+	virtual bool HasKeyboardFocus() const override { return false; }
+	virtual FReply OnFocusReceived(const FGeometry& MyGeometry, const FFocusEvent& InFocusEvent) override;
+	// ~End STreeView
 private:
 	/** The owning tree view widget. */
 	TWeakPtr<SNamingTokenDataTreeViewWidget> OwningTreeViewWidget;
@@ -161,8 +165,12 @@ class SNamingTokenDataTreeViewWidget final : public SCompoundWidget
 	
 public:
 	SLATE_BEGIN_ARGS(SNamingTokenDataTreeViewWidget) {}
+		/** If a selection is changed, such as arrow keys or a single click. */
 		SLATE_EVENT(FOnSelectionChanged, OnSelectionChanged)
+		/** If an item is formally selected, such as by pressing enter or a double click. */
 		SLATE_EVENT(FOnItemSelected, OnItemSelected)
+		/** Event if this widget or its children receives focus. */
+		SLATE_EVENT(FSimpleDelegate, OnFocused)
 	SLATE_END_ARGS()
 
 	virtual ~SNamingTokenDataTreeViewWidget() override;
@@ -192,13 +200,24 @@ public:
 
 	/** Get the namespace filter text applied. */
 	const FString& GetNamespaceFilterText() const { return NamespaceFilterText; }
+
+	/** Informs us we're receiving focus, such as from a child of the treeview. */
+	void NotifyFocusReceived();
+
+	// ~Begin SCompoundWidget
+	virtual FReply OnFocusReceived(const FGeometry& MyGeometry, const FFocusEvent& InFocusEvent) override;
+	// ~End SCompoundWidget
 	
 private:
+	/** Find an item given its index. */
 	FNamingTokenDataTreeItemPtr GetItemFromIndex(int32 InIndex, const TArray<FNamingTokenDataTreeItemPtr>& InItems) const;
-
+	
 	/** Set the currently selected item given an absolute index. */
 	void SetItemFromIndex(int32 InIndex);
 
+	/** Clears the currently selected item. */
+	void ClearCurrentSelection();
+	
 	/** Officially complete and report the selection. Called from tab or double-click. */
 	void FinalizeSelection(FNamingTokenDataTreeItemPtr SelectedItem);
 	
@@ -248,4 +267,7 @@ private:
 
 	/** Delegate for when an item is fully selected. */
 	FOnItemSelected ItemSelectedDelegate;
+
+	/** Delegate if we have received focus. */
+	FSimpleDelegate OnFocusedDelegate;
 };
\ No newline at end of file
diff --git a/Engine/Plugins/Developer/NamingTokens/Source/NamingTokensUI/Private/SNamingTokensEditableTextBox.cpp b/Engine/Plugins/Developer/NamingTokens/Source/NamingTokensUI/Private/SNamingTokensEditableTextBox.cpp
--- a/Engine/Plugins/Developer/NamingTokens/Source/NamingTokensUI/Private/SNamingTokensEditableTextBox.cpp
+++ b/Engine/Plugins/Developer/NamingTokens/Source/NamingTokensUI/Private/SNamingTokensEditableTextBox.cpp
@@ -8,6 +8,7 @@
 #include "Utils/NamingTokenUtils.h"
 
 #include "Engine/Engine.h"
+#include "Framework/Application/SlateApplication.h"
 #include "Internationalization/Regex.h"
 #include "Internationalization/Text.h"
 #include "Layout/Margin.h"
@@ -114,10 +115,13 @@ void SNamingTokensEditableTextBox::Construct(const FArguments& InArgs)
 				.FillHeight(1.0f)
 				[
 					SAssignNew(DropdownAnchor, SMenuAnchor)
-					.Method(EPopupMethod::CreateNewWindow)
+					.Method(EPopupMethod::UseCurrentWindow)
 					.Placement(MenuPlacement_ComboBox)
 					// The drop down filter menu.
 					.OnGetMenuContent(this, &SNamingTokensEditableTextBox::OnGetDropdownContent)
+					// False to prevent the menu from being dismissed by the application. Otherwise, it gets dismissed
+					// when the user clicks on the dropdown, because we send a keyboard focus event back to the text widget.
+					.UseApplicationMenuStack(false)
 					.Visibility(EVisibility::Visible)
 					.Content()
 					[
@@ -337,7 +341,7 @@ bool SNamingTokensEditableTextBox::HasKeyboardFocus() const
 FReply SNamingTokensEditableTextBox::OnFocusReceived(const FGeometry& MyGeometry, const FFocusEvent& InFocusEvent)
 {
 	FReply Reply = FReply::Handled();
-	if ( InFocusEvent.GetCause() != EFocusCause::Cleared )
+	if (InFocusEvent.GetCause() != EFocusCause::Cleared)
 	{
 		// Forward keyboard focus to our editable text widget
 		Reply.SetUserFocus(EditableText.ToSharedRef(), InFocusEvent.GetCause());
@@ -496,9 +500,9 @@ void SNamingTokensEditableTextBox::OnEditableTextCommitted(const FText& InTokeni
 {
 	if (DropdownAnchor.IsValid())
 	{
-		const EFocusCause FocusType = DropdownAnchor->HasAnyUserFocus().Get(EFocusCause::Cleared);
+		const EFocusCause DropdownFocusType = DropdownAnchor->HasAnyUserFocus().Get(EFocusCause::Cleared);
 
-		if (InCommitType != ETextCommit::OnUserMovedFocus && FocusType == EFocusCause::Cleared)
+		if (InCommitType != ETextCommit::OnUserMovedFocus && DropdownFocusType == EFocusCause::Cleared)
 		{
 			CloseSuggestionDropdown();
 		}
@@ -561,7 +565,6 @@ void SNamingTokensEditableTextBox::OpenOrCloseSuggestionDropdown() const
 		}
 		else
 		{
-			LastCursorTokenData.Reset();
 			CloseSuggestionDropdown();
 		}
 	}
@@ -594,10 +597,11 @@ void SNamingTokensEditableTextBox::OpenSuggestionDropdown() const
 
 void SNamingTokensEditableTextBox::CloseSuggestionDropdown() const
 {
-	if (DropdownAnchor.IsValid())
+	if (DropdownAnchor.IsValid() && DropdownAnchor->IsOpen())
 	{
 		DropdownAnchor->SetIsOpen(false, false);
 	}
+	LastCursorTokenData.Reset();
 }
 
 TSharedPtr<SNamingTokensEditableTextBox::FCursorTokenData> SNamingTokensEditableTextBox::GetCurrentCursorTokenData() const
@@ -677,7 +681,8 @@ TSharedRef<SWidget> SNamingTokensEditableTextBox::OnGetDropdownContent()
 		return SNullWidget::NullWidget;
 	}
 	TSharedRef<SNamingTokenDataTreeViewWidget> Widget = SAssignNew(FilterWidget, SNamingTokenDataTreeViewWidget)
-		.OnItemSelected(this, &SNamingTokensEditableTextBox::OnSuggestionSelected);
+		.OnItemSelected(this, &SNamingTokensEditableTextBox::OnSuggestionDoubleClicked)
+		.OnFocused(this, &SNamingTokensEditableTextBox::OnSuggestionDropdownReceivedFocus);
 	return Widget;
 }
 
@@ -801,6 +806,30 @@ void SNamingTokensEditableTextBox::OnSuggestionSelected(TSharedPtr<FNamingTokenD
 	}
 }
 
+void SNamingTokensEditableTextBox::OnSuggestionDoubleClicked(TSharedPtr<FNamingTokenDataTreeItem> SelectedItem)
+{
+	// Insert the selected item to our text box.
+	OnSuggestionSelected(SelectedItem);
+	// Gather the plain text from the text box lines, because GetText() returns OUR binding, which hasn't been updated yet.
+	const FText EnteredText = EditableText->GetPlainText();
+	// A double click from the suggestion means the suggestion box had focus, and OUR text commit event has already fired. We need to
+	// run our commit logic to update bindings and validate the token.
+	OnEditableTextCommitted(EnteredText, ETextCommit::Type::OnUserMovedFocus);
+
+	// Focus the text box so the user can keep typing. Chances are a double click on the suggestion box doesn't mean
+	// the user wanted text input to stop.
+	FSlateApplication::Get().SetKeyboardFocus(EditableText, EFocusCause::OtherWidgetLostFocus);
+	FSlateApplication::Get().SetUserFocus(0, EditableText, EFocusCause::OtherWidgetLostFocus);
+}
+
+void SNamingTokensEditableTextBox::OnSuggestionDropdownReceivedFocus() const
+{
+	// Without this focus event, the keyboard focus will remain on the dropdown and prevent the user from typing,
+	// moving the cursor, or pressing tab or enter to confirm a selection. Keyboard navigation of the dropdown
+	// is handled separately and requires our editable text box to maintain focus.
+	FSlateApplication::Get().SetKeyboardFocus(EditableText, EFocusCause::WindowActivate);
+}
+
 END_SLATE_FUNCTION_BUILD_OPTIMIZATION
 
 #undef LOCTEXT_NAMESPACE
\ No newline at end of file
diff --git a/Engine/Plugins/Developer/NamingTokens/Source/NamingTokensUI/Public/SNamingTokensEditableTextBox.h b/Engine/Plugins/Developer/NamingTokens/Source/NamingTokensUI/Public/SNamingTokensEditableTextBox.h
--- a/Engine/Plugins/Developer/NamingTokens/Source/NamingTokensUI/Public/SNamingTokensEditableTextBox.h
+++ b/Engine/Plugins/Developer/NamingTokens/Source/NamingTokensUI/Public/SNamingTokensEditableTextBox.h
@@ -17,6 +17,7 @@ class FText;
 class SMenuAnchor;
 
 struct FNamingTokenFilterArgs;
+struct FNamingTokenDataTreeItem;
 
 /**
  * An editable text box for displaying tokenized strings in either their unevaluated or resolved form.
@@ -200,7 +201,13 @@ private:
 	void OnEditableTextCursorMoved(const FTextLocation& NewCursorPosition);
 	
 	/** When a dropdown suggestion has been selected. */
-	void OnSuggestionSelected(TSharedPtr<struct FNamingTokenDataTreeItem> SelectedItem) const;
+	void OnSuggestionSelected(TSharedPtr<FNamingTokenDataTreeItem> SelectedItem) const;
+
+	/** Called when the suggestion box has been double clicked. */
+	void OnSuggestionDoubleClicked(TSharedPtr<FNamingTokenDataTreeItem> SelectedItem);
+
+	/** Called if the suggestion box has received focus. */
+	void OnSuggestionDropdownReceivedFocus() const;
 	
 private:
 	/** Filter args to use during evaluation. */
