NamingTokens: Fix slate widget not initializing contexts to nullptr.

diff --git a/Engine/Plugins/Developer/NamingTokens/Source/NamingTokensUI/Public/SNamingTokensEditableTextBox.h b/Engine/Plugins/Developer/NamingTokens/Source/NamingTokensUI/Public/SNamingTokensEditableTextBox.h
--- a/Engine/Plugins/Developer/NamingTokens/Source/NamingTokensUI/Public/SNamingTokensEditableTextBox.h
+++ b/Engine/Plugins/Developer/NamingTokens/Source/NamingTokensUI/Public/SNamingTokensEditableTextBox.h
@@ -35,6 +35,7 @@ public:
 		, _DisplayErrorMessage(true)
 		, _DisplayBorderImage(true)
 		, _NamingTokenFilterArgs(nullptr)
+		, _Contexts(nullptr)
 		, _EnableSuggestionDropdown(true)
 		, _FullyQualifyGlobalTokenSuggestions(false)
 	{}
