NamingTokens: Slate Widget - Change Contexts and FilterArgs params to be passed by value. Make EvaluateNamingTokens public.

diff --git a/Engine/Plugins/Developer/NamingTokens/Source/NamingTokensUI/Private/NamingTokensEditableText.cpp b/Engine/Plugins/Developer/NamingTokens/Source/NamingTokensUI/Private/NamingTokensEditableText.cpp
--- a/Engine/Plugins/Developer/NamingTokens/Source/NamingTokensUI/Private/NamingTokensEditableText.cpp
+++ b/Engine/Plugins/Developer/NamingTokens/Source/NamingTokensUI/Private/NamingTokensEditableText.cpp
@@ -129,8 +129,8 @@ TSharedRef<SWidget> UNamingTokensEditableText::RebuildWidget()
 		.AllowMultiLine(bIsMultiline)
 		.Style(&*EditableTextBoxStyle)
 		.ArgumentStyle(&ArgumentStyle)
-		.NamingTokenFilterArgs(&FilterArgs)
-		.Contexts(&Contexts)
+		.FilterArgs(FilterArgs)
+		.Contexts(Contexts)
 		.EnableSuggestionDropdown(bEnableSuggestionDropdown)
 		.DisplayTokenIcon_UObject(this, &UNamingTokensEditableText::GetDisplayTokenIcon)
 		.DisplayErrorMessage_UObject(this, &UNamingTokensEditableText::GetDisplayErrorMessage)
@@ -203,7 +203,7 @@ void UNamingTokensEditableText::OnPreEvaluateNamingTokensCallback()
 	
 	if (MyNamingTokensEditableText.IsValid())
 	{
-		MyNamingTokensEditableText->SetFilterArgs(&FilterArgs);
+		MyNamingTokensEditableText->SetFilterArgs(FilterArgs);
 		MyNamingTokensEditableText->SetContexts(Contexts);
 	}
 }
diff --git a/Engine/Plugins/Developer/NamingTokens/Source/NamingTokensUI/Private/SNamingTokensEditableTextBox.cpp b/Engine/Plugins/Developer/NamingTokens/Source/NamingTokensUI/Private/SNamingTokensEditableTextBox.cpp
--- a/Engine/Plugins/Developer/NamingTokens/Source/NamingTokensUI/Private/SNamingTokensEditableTextBox.cpp
+++ b/Engine/Plugins/Developer/NamingTokens/Source/NamingTokensUI/Private/SNamingTokensEditableTextBox.cpp
@@ -2,7 +2,6 @@
 
 #include "SNamingTokensEditableTextBox.h"
 
-#include "NamingTokensEngineSubsystem.h"
 #include "NamingTokensStringSyntaxHighlighterMarshaller.h"
 #include "SNamingTokensDataTreeView.h"
 #include "Utils/NamingTokenUtils.h"
@@ -67,8 +66,8 @@ void SNamingTokensEditableTextBox::Construct(const FArguments& InArgs)
 	bEnableSuggestionDropdown = InArgs._EnableSuggestionDropdown;
 	bFullyQualifyGlobalTokenSuggestions = InArgs._FullyQualifyGlobalTokenSuggestions;
 	
-	SetFilterArgs(InArgs._NamingTokenFilterArgs);
-	SetContexts(InArgs._Contexts ? *InArgs._Contexts : TArray<UObject*>());
+	SetFilterArgs(InArgs._FilterArgs);
+	SetContexts(InArgs._Contexts);
 	
 	ErrorReporting = SNew(SPopupErrorText);
 	ErrorReporting->SetError(FText::GetEmpty());
@@ -248,16 +247,9 @@ void SNamingTokensEditableTextBox::SetEnableSuggestionDropdown(bool bInEnableSug
 	bEnableSuggestionDropdown = bInEnableSuggestionDropdown;
 }
 
-void SNamingTokensEditableTextBox::SetFilterArgs(const FNamingTokenFilterArgs* InFilterArgs)
+void SNamingTokensEditableTextBox::SetFilterArgs(const FNamingTokenFilterArgs& InFilterArgs)
 {
-	if (InFilterArgs)
-	{
-		NamingTokenFilterArgs = MakeShared<FNamingTokenFilterArgs>(*InFilterArgs);
-	}
-	else
-	{
-		NamingTokenFilterArgs.Reset();
-	}
+	NamingTokenFilterArgs = InFilterArgs;
 }
 
 void SNamingTokensEditableTextBox::SetContexts(const TArray<UObject*>& InContexts)
@@ -265,6 +257,19 @@ void SNamingTokensEditableTextBox::SetContexts(const TArray<UObject*>& InContext
 	NamingTokenContexts = InContexts;
 }
 
+void SNamingTokensEditableTextBox::EvaluateNamingTokens()
+{
+	check(bHandlingEvaluation);
+
+	OnPreEvaluateTokens.ExecuteIfBound();
+	
+	UNamingTokensEngineSubsystem* TokenSubsystem = GEngine->GetEngineSubsystem<UNamingTokensEngineSubsystem>();
+	const FNamingTokenResultData EvaluationData = TokenSubsystem->EvaluateTokenText(TokenizedText.Get(),
+		NamingTokenFilterArgs, NamingTokenContexts);
+
+	ResolvedText.Set(EvaluationData.EvaluatedText);
+}
+
 FText SNamingTokensEditableTextBox::GetText() const
 {
 	if (HasKeyboardFocus())
@@ -536,19 +541,6 @@ void SNamingTokensEditableTextBox::UpdateFilters() const
 	}
 }
 
-void SNamingTokensEditableTextBox::EvaluateNamingTokens()
-{
-	check(bHandlingEvaluation);
-
-	OnPreEvaluateTokens.ExecuteIfBound();
-	
-	UNamingTokensEngineSubsystem* TokenSubsystem = GEngine->GetEngineSubsystem<UNamingTokensEngineSubsystem>();
-	const FNamingTokenResultData EvaluationData = TokenSubsystem->EvaluateTokenText(TokenizedText.Get(),
-		NamingTokenFilterArgs.IsValid() ? *NamingTokenFilterArgs : FNamingTokenFilterArgs(), NamingTokenContexts);
-
-	ResolvedText.Set(EvaluationData.EvaluatedText);
-}
-
 void SNamingTokensEditableTextBox::OpenOrCloseSuggestionDropdown() const
 {
 	if (DropdownAnchor.IsValid())
diff --git a/Engine/Plugins/Developer/NamingTokens/Source/NamingTokensUI/Public/SNamingTokensEditableTextBox.h b/Engine/Plugins/Developer/NamingTokens/Source/NamingTokensUI/Public/SNamingTokensEditableTextBox.h
--- a/Engine/Plugins/Developer/NamingTokens/Source/NamingTokensUI/Public/SNamingTokensEditableTextBox.h
+++ b/Engine/Plugins/Developer/NamingTokens/Source/NamingTokensUI/Public/SNamingTokensEditableTextBox.h
@@ -6,6 +6,7 @@
 #include "CoreMinimal.h"
 #include "Framework/SlateDelegates.h"
 #include "Misc/Attribute.h"
+#include "NamingTokensEngineSubsystem.h"
 #include "Templates/SharedPointer.h"
 #include "Widgets/DeclarativeSyntaxSupport.h"
 #include "Widgets/Input/SMultiLineEditableTextBox.h"
@@ -16,7 +17,6 @@ class FString;
 class FText;
 class SMenuAnchor;
 
-struct FNamingTokenFilterArgs;
 struct FNamingTokenDataTreeItem;
 
 /**
@@ -34,8 +34,6 @@ public:
 		, _DisplayTokenIcon(true)
 		, _DisplayErrorMessage(true)
 		, _DisplayBorderImage(true)
-		, _NamingTokenFilterArgs(nullptr)
-		, _Contexts(nullptr)
 		, _EnableSuggestionDropdown(true)
 		, _FullyQualifyGlobalTokenSuggestions(false)
 	{}
@@ -88,10 +86,10 @@ public:
 		SLATE_EVENT(FSimpleDelegate, OnPreEvaluateNamingTokens)
 
 		/** Filter args to apply during evaluation. */
-		SLATE_ARGUMENT(const FNamingTokenFilterArgs*, NamingTokenFilterArgs)
+		SLATE_ARGUMENT(FNamingTokenFilterArgs, FilterArgs)
 
 		/** Contexts to provide to evaluation. */
-		SLATE_ARGUMENT(const TArray<TObjectPtr<UObject>>*, Contexts)
+		SLATE_ARGUMENT(TArray<TObjectPtr<UObject>>, Contexts)
 
 		/** If drop down suggestions should be enabled. */
 		SLATE_ARGUMENT(bool, EnableSuggestionDropdown)
@@ -129,10 +127,13 @@ public:
 	NAMINGTOKENSUI_API void SetEnableSuggestionDropdown(bool bInEnableSuggestionDropdown);
 	
 	/** Provide filter args to use during evaluation. Passing null clears existing. */
-	NAMINGTOKENSUI_API void SetFilterArgs(const FNamingTokenFilterArgs* InFilterArgs);
+	NAMINGTOKENSUI_API void SetFilterArgs(const FNamingTokenFilterArgs& InFilterArgs);
 
 	/** Provide contexts to naming token evaluation. */
 	NAMINGTOKENSUI_API void SetContexts(const TArray<UObject*>& InContexts);
+
+	/** Evaluate the token text and update the resolved text. */
+	NAMINGTOKENSUI_API void EvaluateNamingTokens();
 	
 private:
 	/** @return Text contents based on current state */
@@ -161,8 +162,6 @@ private:
 
 	/** Check for type text and update filters if needed. */
 	void UpdateFilters() const;
-	
-	void EvaluateNamingTokens();
 
 	struct FCursorTokenData
 	{
@@ -212,7 +211,7 @@ private:
 	
 private:
 	/** Filter args to use during evaluation. */
-	TSharedPtr<FNamingTokenFilterArgs> NamingTokenFilterArgs;
+	FNamingTokenFilterArgs NamingTokenFilterArgs;
 
 	/** Contexts to provide to evaluation. */
 	TArray<TObjectPtr<UObject>> NamingTokenContexts;
