NamingTokens: Fix doc string.

diff --git a/Engine/Plugins/Developer/NamingTokens/Source/NamingTokensUI/Public/SNamingTokensEditableTextBox.h b/Engine/Plugins/Developer/NamingTokens/Source/NamingTokensUI/Public/SNamingTokensEditableTextBox.h
--- a/Engine/Plugins/Developer/NamingTokens/Source/NamingTokensUI/Public/SNamingTokensEditableTextBox.h
+++ b/Engine/Plugins/Developer/NamingTokens/Source/NamingTokensUI/Public/SNamingTokensEditableTextBox.h
@@ -126,7 +126,7 @@ public:
 	/** If the suggestion box should be enabled. */
 	NAMINGTOKENSUI_API void SetEnableSuggestionDropdown(bool bInEnableSuggestionDropdown);
 	
-	/** Provide filter args to use during evaluation. Passing null clears existing. */
+	/** Provide filter args to use during evaluation. */
 	NAMINGTOKENSUI_API void SetFilterArgs(const FNamingTokenFilterArgs& InFilterArgs);
 
 	/** Provide contexts to naming token evaluation. */
