Naming Tokens: Change UMG display name to "Naming Tokens Editable Text Box" making it more consistent with the C++ class and general text naming.

diff --git a/Engine/Plugins/Developer/NamingTokens/Source/NamingTokensUI/Public/NamingTokensEditableText.h b/Engine/Plugins/Developer/NamingTokens/Source/NamingTokensUI/Public/NamingTokensEditableText.h
--- a/Engine/Plugins/Developer/NamingTokens/Source/NamingTokensUI/Public/NamingTokensEditableText.h
+++ b/Engine/Plugins/Developer/NamingTokens/Source/NamingTokensUI/Public/NamingTokensEditableText.h
@@ -13,7 +13,7 @@ struct FEditableTextBoxStyle;
 /**
  * Editable text for displaying tokenized strings in either their unevaluated or resolved form.
  */
-UCLASS(MinimalAPI, Config = NamingTokens, DefaultConfig, ClassGroup = UI, meta = (Category = "NamingTokens UI", DisplayName = "Naming Tokens Text Block", PrioritizeCategories = "Content"))
+UCLASS(MinimalAPI, Config = NamingTokens, DefaultConfig, ClassGroup = UI, meta = (Category = "NamingTokens UI", DisplayName = "Naming Tokens Editable Text Box", PrioritizeCategories = "Content"))
 class UNamingTokensEditableText : public UMultiLineEditableText
 {
 	GENERATED_BODY()
