NamingTokens: Add OnTokenTextEvaluated event to the token textbox widget that executes every time the token text is re-evaluated, providing the new resolved text.
Add a ShouldEvaluateTokens attribute that determines whether the widget should try to resolve the tokenized string or not.
Add an EvaluationFrequency attribute that, if set, will be used with a Slate timer to periodically re-evaluate the token text.

#rnx

diff --git a/Engine/Plugins/Developer/NamingTokens/Source/NamingTokensUI/Private/SNamingTokensEditableTextBox.cpp b/Engine/Plugins/Developer/NamingTokens/Source/NamingTokensUI/Private/SNamingTokensEditableTextBox.cpp
--- a/Engine/Plugins/Developer/NamingTokens/Source/NamingTokensUI/Private/SNamingTokensEditableTextBox.cpp
+++ b/Engine/Plugins/Developer/NamingTokens/Source/NamingTokensUI/Private/SNamingTokensEditableTextBox.cpp
@@ -57,6 +57,10 @@ void SNamingTokensEditableTextBox::Construct(const FArguments& InArgs)
 	OnTokenizedTextChanged = InArgs._OnTextChanged;
 	OnTokenizedTextCommitted = InArgs._OnTextCommitted;
 	OnPreEvaluateTokens = InArgs._OnPreEvaluateNamingTokens;
+	OnTokenizedTextEvaluated = InArgs._OnTokenizedTextEvaluated;
+
+	ShouldEvaluateTokens = InArgs._ShouldEvaluateTokens;
+	EvaluationFrequency = InArgs._EvaluationFrequency;
 	CanDisplayResolvedText = InArgs._CanDisplayResolvedText;
 
 	DisplayTokenIcon = InArgs._DisplayTokenIcon;
@@ -72,10 +76,20 @@ void SNamingTokensEditableTextBox::Construct(const FArguments& InArgs)
 	ErrorReporting = SNew(SPopupErrorText);
 	ErrorReporting->SetError(FText::GetEmpty());
 
-	bHandlingEvaluation = !ResolvedText.IsSet();
-	
+	// Create a Slate timer that will re-evaluate the token text at the specified evaluation frequency
+	if (EvaluationFrequency.Get() > 0.0f)
+	{
+		RegisterActiveTimer(EvaluationFrequency.Get(), FWidgetActiveTimerDelegate::CreateLambda([this](double InCurrentTime, float InDeltaTime) -> EActiveTimerReturnType
+			{
+				EvaluateNamingTokens();
+				return EActiveTimerReturnType::Continue;
+			}));
+	}
+
 	BoundText = MakeAttributeSP(this, &SNamingTokensEditableTextBox::GetText);
 	
+	const FMargin BorderPadding = InArgs._IsReadOnly.Get() ? FMargin(8.0f, 0.0f) : TextBoxStyle->Padding;
+
 	ChildSlot
 	[
 		SNew(SBorder)
@@ -83,7 +97,7 @@ void SNamingTokensEditableTextBox::Construct(const FArguments& InArgs)
 		.BorderBackgroundColor(this, &SNamingTokensEditableTextBox::GetBorderBackgroundColor)
 		.ForegroundColor(this, &SMultiLineEditableTextBox::GetForegroundColor)
 		.VAlign(VAlign_Center)
-		.Padding(TextBoxStyle->Padding)
+		.Padding(BorderPadding)
 		[
 			SAssignNew(Box, SHorizontalBox)
 
@@ -91,7 +105,7 @@ void SNamingTokensEditableTextBox::Construct(const FArguments& InArgs)
 			.HAlign(HAlign_Fill)
 			.VAlign(VAlign_Center)
 			.AutoWidth()
-			.Padding(4, 0)
+			.Padding(0, 0, 4, 0)
 			[
 				SNew(SImage)
 					.Image(FAppStyle::Get().GetBrush("Icons.TokenTextBox"))
@@ -183,10 +197,7 @@ void SNamingTokensEditableTextBox::Construct(const FArguments& InArgs)
 		]
 	];
 	
-	if (bHandlingEvaluation)
-	{
-		OnEditableTextCommitted(TokenizedText.Get(), ETextCommit::Default);
-	}
+	OnEditableTextCommitted(TokenizedText.Get(), ETextCommit::Default);
 }
 
 void SNamingTokensEditableTextBox::SetTokenizedText(const FText& InText)
@@ -259,15 +270,18 @@ void SNamingTokensEditableTextBox::SetContexts(const TArray<UObject*>& InContext
 
 void SNamingTokensEditableTextBox::EvaluateNamingTokens()
 {
-	check(bHandlingEvaluation);
+	if (ShouldEvaluateTokens.Get() && !ResolvedText.IsBound())
+	{
+		OnPreEvaluateTokens.ExecuteIfBound();
 
-	OnPreEvaluateTokens.ExecuteIfBound();
-	
-	UNamingTokensEngineSubsystem* TokenSubsystem = GEngine->GetEngineSubsystem<UNamingTokensEngineSubsystem>();
-	const FNamingTokenResultData EvaluationData = TokenSubsystem->EvaluateTokenText(TokenizedText.Get(),
-		NamingTokenFilterArgs, NamingTokenContexts);
+		UNamingTokensEngineSubsystem* TokenSubsystem = GEngine->GetEngineSubsystem<UNamingTokensEngineSubsystem>();
+		const FNamingTokenResultData EvaluationData = TokenSubsystem->EvaluateTokenText(TokenizedText.Get(),
+			NamingTokenFilterArgs, NamingTokenContexts);
+
+		ResolvedText.Set(EvaluationData.EvaluatedText);
 
-	ResolvedText.Set(EvaluationData.EvaluatedText);
+		OnTokenizedTextEvaluated.ExecuteIfBound(ResolvedText.Get());
+	}
 }
 
 FText SNamingTokensEditableTextBox::GetText() const
@@ -522,10 +536,7 @@ void SNamingTokensEditableTextBox::OnEditableTextCommitted(const FText& InTokeni
 
 	OnTokenizedTextCommitted.ExecuteIfBound(InTokenizedText, InCommitType);
 
-	if (bHandlingEvaluation)
-	{
-		EvaluateNamingTokens();
-	}
+	EvaluateNamingTokens();
 }
 
 void SNamingTokensEditableTextBox::UpdateFilters() const
diff --git a/Engine/Plugins/Developer/NamingTokens/Source/NamingTokensUI/Public/SNamingTokensEditableTextBox.h b/Engine/Plugins/Developer/NamingTokens/Source/NamingTokensUI/Public/SNamingTokensEditableTextBox.h
--- a/Engine/Plugins/Developer/NamingTokens/Source/NamingTokensUI/Public/SNamingTokensEditableTextBox.h
+++ b/Engine/Plugins/Developer/NamingTokens/Source/NamingTokensUI/Public/SNamingTokensEditableTextBox.h
@@ -19,6 +19,8 @@ class SMenuAnchor;
 
 struct FNamingTokenDataTreeItem;
 
+DECLARE_DELEGATE_OneParam(FOnTokenizedTextEvaluated, const FText&);
+
 /**
  * An editable text box for displaying tokenized strings in either their unevaluated or resolved form.
  */
@@ -28,6 +30,8 @@ public:
 	SLATE_BEGIN_ARGS(SNamingTokensEditableTextBox)
 		: _Style(nullptr)
 		, _ArgumentStyle(nullptr)
+		, _ShouldEvaluateTokens(true)
+		, _EvaluationFrequency(0.0f)
 		, _AllowMultiLine(false)
 		, _IsReadOnly(false)
 		, _CanDisplayResolvedText(true)
@@ -55,6 +59,12 @@ public:
 		/** The marshaller used to get/set the raw text to/from the text layout. */
 		SLATE_ARGUMENT(TSharedPtr<ITextLayoutMarshaller>, Marshaller)
 
+		/** Whether this widget should evaluate the tokenized string */
+		SLATE_ATTRIBUTE(bool, ShouldEvaluateTokens)
+
+		/** If set, specifies how frequently the tokenized text should be re-evaluated */
+		SLATE_ATTRIBUTE(float, EvaluationFrequency)
+
 		/** Whether to allow multi-line text. */
 		SLATE_ATTRIBUTE(bool, AllowMultiLine)
 
@@ -85,6 +95,9 @@ public:
 		/** Called prior to Naming Token evaluation. Only fires if this widget is handling evaluations. */
 		SLATE_EVENT(FSimpleDelegate, OnPreEvaluateNamingTokens)
 
+		/** Called after the tokenized text is evaluated to provide the resulting resolved text. */
+		SLATE_EVENT(FOnTokenizedTextEvaluated, OnTokenizedTextEvaluated)
+
 		/** Filter args to apply during evaluation. */
 		SLATE_ARGUMENT(FNamingTokenFilterArgs, FilterArgs)
 
@@ -252,6 +265,12 @@ private:
 	TAttribute<FText> TokenizedText;
 	TAttribute<FText> ResolvedText;
 
+	/** If true, this widget will evaluate the tokenized text to produce the resolved text */
+	TAttribute<bool> ShouldEvaluateTokens;
+
+	/** If set, specifies how frequently the tokenized text should be re-evaluated */
+	TAttribute<float> EvaluationFrequency;
+
 	/** If we are displaying the token icon. */
 	TAttribute<bool> DisplayTokenIcon;
 
@@ -277,8 +296,8 @@ private:
 	/** Callback before we evaluate tokens. */
 	FSimpleDelegate OnPreEvaluateTokens;
 
-	/** If this widget is handling all token evaluation. */
-	bool bHandlingEvaluation = false;
+	/** Callback after the tokenized text has been evaluated */
+	FOnTokenizedTextEvaluated OnTokenizedTextEvaluated;
 
 	/** If global tokens should have their namespace inserted too from suggestions. */
 	bool bFullyQualifyGlobalTokenSuggestions = false;
