NamingTokens: Changes the behavior of tokenized text validation so that it now prevents committing the text when invalid. This more closely matches the behavior of other editable textbox widgets.
Added an additional callback to allow for validation of the resolved text. Resolved text validation does not prevent committing tokenized text, but does allow for error messages to be displayed when invalid.
Added support to display the textbox widget like an inline editable textbox. In other words, when enabled via a widget attribute, the widget will appear like a read-only textblock, with no token icon or background image, until edit mode is triggered, which will show the background image and token icon, allow the user to edit the text, and then revert back to the read-only appearance after the text has been committed.
Added textbox padding override attribute.
Minor padding / UI tweaks to more closely resemble other textbox widgets.

#rnx

diff --git a/Engine/Plugins/Developer/NamingTokens/Source/NamingTokensUI/Private/SNamingTokensEditableTextBox.cpp b/Engine/Plugins/Developer/NamingTokens/Source/NamingTokensUI/Private/SNamingTokensEditableTextBox.cpp
--- a/Engine/Plugins/Developer/NamingTokens/Source/NamingTokensUI/Private/SNamingTokensEditableTextBox.cpp
+++ b/Engine/Plugins/Developer/NamingTokens/Source/NamingTokensUI/Private/SNamingTokensEditableTextBox.cpp
@@ -54,6 +54,7 @@ void SNamingTokensEditableTextBox::Construct(const FArguments& InArgs)
 	ValidArguments = InArgs._ValidArguments;
 
 	OnValidateTokenizedText = InArgs._OnValidateTokenizedText;
+	OnValidateResolvedText = InArgs._OnValidateResolvedText;
 	OnTokenizedTextChanged = InArgs._OnTextChanged;
 	OnTokenizedTextCommitted = InArgs._OnTextCommitted;
 	OnPreEvaluateTokens = InArgs._OnPreEvaluateNamingTokens;
@@ -66,7 +67,9 @@ void SNamingTokensEditableTextBox::Construct(const FArguments& InArgs)
 	DisplayTokenIcon = InArgs._DisplayTokenIcon;
 	DisplayErrorMessage = InArgs._DisplayErrorMessage;
 	DisplayBorderImage = InArgs._DisplayBorderImage;
+	TextBoxPaddingOverride = InArgs._TextBoxPadding;
 
+	bSupportInlineEdit = InArgs._SupportInlineEdit;
 	bEnableSuggestionDropdown = InArgs._EnableSuggestionDropdown;
 	bFullyQualifyGlobalTokenSuggestions = InArgs._FullyQualifyGlobalTokenSuggestions;
 	
@@ -88,7 +91,7 @@ void SNamingTokensEditableTextBox::Construct(const FArguments& InArgs)
 
 	BoundText = MakeAttributeSP(this, &SNamingTokensEditableTextBox::GetText);
 	
-	const FMargin BorderPadding = InArgs._IsReadOnly.Get() ? FMargin(8.0f, 0.0f) : TextBoxStyle->Padding;
+	const FMargin TextBoxPadding = TextBoxPaddingOverride.IsSet() ? TextBoxPaddingOverride.Get() : TextBoxStyle->Padding;
 
 	ChildSlot
 	[
@@ -96,8 +99,7 @@ void SNamingTokensEditableTextBox::Construct(const FArguments& InArgs)
 		.BorderImage(this, &SNamingTokensEditableTextBox::GetBorderImage)
 		.BorderBackgroundColor(this, &SNamingTokensEditableTextBox::GetBorderBackgroundColor)
 		.ForegroundColor(this, &SMultiLineEditableTextBox::GetForegroundColor)
-		.VAlign(VAlign_Center)
-		.Padding(BorderPadding)
+		.Padding(0)
 		[
 			SAssignNew(Box, SHorizontalBox)
 
@@ -105,7 +107,7 @@ void SNamingTokensEditableTextBox::Construct(const FArguments& InArgs)
 			.HAlign(HAlign_Fill)
 			.VAlign(VAlign_Center)
 			.AutoWidth()
-			.Padding(0, 0, 4, 0)
+			.Padding(4, 0, 0, 0)
 			[
 				SNew(SImage)
 					.Image(FAppStyle::Get().GetBrush("Icons.TokenTextBox"))
@@ -125,6 +127,7 @@ void SNamingTokensEditableTextBox::Construct(const FArguments& InArgs)
 				+ SVerticalBox::Slot()
 				.HAlign(HAlign_Fill)
 				.VAlign(VAlign_Fill)
+				.Padding(TextBoxPadding)
 				.FillHeight(1.0f)
 				[
 					SAssignNew(DropdownAnchor, SMenuAnchor)
@@ -253,6 +256,11 @@ void SNamingTokensEditableTextBox::SetArgumentStyle(const FTextBlockStyle* InArg
 	}
 }
 
+void SNamingTokensEditableTextBox::SetSupportInlineEdit(bool bInSupportInlineEdit)
+{
+	bSupportInlineEdit = bInSupportInlineEdit;
+}
+
 void SNamingTokensEditableTextBox::SetEnableSuggestionDropdown(bool bInEnableSuggestionDropdown)
 {
 	bEnableSuggestionDropdown = bInEnableSuggestionDropdown;
@@ -280,6 +288,17 @@ void SNamingTokensEditableTextBox::EvaluateNamingTokens()
 
 		ResolvedText.Set(EvaluationData.EvaluatedText);
 
+		if (OnValidateResolvedText.IsBound())
+		{
+			FText ErrorMessage;
+			const bool bIsValid = OnValidateResolvedText.Execute(ResolvedText.Get(), ErrorMessage);
+
+			if (!HasKeyboardFocus())
+			{
+				SetError(bIsValid ? FText::GetEmpty() : ErrorMessage);
+			}
+		}
+
 		OnTokenizedTextEvaluated.ExecuteIfBound(ResolvedText.Get());
 	}
 }
@@ -351,6 +370,11 @@ void SNamingTokensEditableTextBox::SetError(const FText& InError) const
 	ErrorReporting->SetError(InError);
 }
 
+bool SNamingTokensEditableTextBox::SupportsKeyboardFocus() const
+{
+	return true;
+}
+
 bool SNamingTokensEditableTextBox::HasKeyboardFocus() const
 {
 	// Since keyboard focus is forwarded to our editable text, we will test it instead
@@ -392,36 +416,55 @@ FReply SNamingTokensEditableTextBox::OnKeyDown(const FGeometry& MyGeometry, cons
 	return FReply::Unhandled();
 }
 
-bool SNamingTokensEditableTextBox::ValidateTokenizedText(const FText& InTokenizedText)
+void SNamingTokensEditableTextBox::EnterEditingMode()
 {
-	FTextBuilder ErrorMessageBuilder;
-	FText ErrorMessage;
-
-	if (!ValidateTokenBraces(InTokenizedText, ErrorMessage))
+	if (bSupportInlineEdit && !EditableText->IsTextReadOnly())
 	{
-		ErrorMessageBuilder.AppendLine(ErrorMessage);
+		FSlateApplication::Get().SetKeyboardFocus(EditableText.ToSharedRef());
+		EditableText->SelectAllText();
+
+		EditableText->SetVisibility(EVisibility::Visible);
+		DisplayTokenIcon = true;
+		DisplayBorderImage = true;
 	}
+}
 
-	// Only validate tokens if there were no other errors - it's expensive
-	if (ErrorMessage.IsEmpty() && !ValidateTokenArgs(InTokenizedText, ErrorMessage))
+void SNamingTokensEditableTextBox::ExitEditingMode()
+{
+	if (bSupportInlineEdit)
 	{
-		ErrorMessageBuilder.AppendLine(ErrorMessage);
+		// Update the visibility to more closely match that of a normal read-only text block
+		EditableText->SetVisibility(EVisibility::HitTestInvisible);
+		DisplayTokenIcon = false;
+		DisplayBorderImage = false;
 	}
+}
+
+bool SNamingTokensEditableTextBox::ValidateTokenizedText(const FText& InTokenizedText)
+{
+	FText ErrorMessage;
 
-	if (!ErrorMessageBuilder.IsEmpty())
+	if (!ValidateTokenArgs(InTokenizedText, ErrorMessage))
 	{
-		ErrorMessage = ErrorMessageBuilder.ToText();
+		SetError(ErrorMessage);
+		return false;
 	}
 
-	if (!ErrorMessage.IsEmpty()	&& OnValidateTokenizedText.IsBound())
+	if (OnValidateTokenizedText.IsBound() && !OnValidateTokenizedText.Execute(InTokenizedText, ErrorMessage))
 	{
-		OnValidateTokenizedText.Execute(InTokenizedText, ErrorMessage);
+		SetError(ErrorMessage);
+		return false;
 	}
 
-	SetError(ErrorMessage);
+	// Validate this last because unbalanced braces should not prevent the text from being committed
+	if (!ValidateTokenBraces(InTokenizedText, ErrorMessage))
+	{
+		SetError(ErrorMessage);
+		return true;
+	}
 
-	// If the error message is empty, the tokenized text is valid
-	return ErrorMessage.IsEmpty();
+	SetError(FText::GetEmpty());
+	return true;
 }
 
 bool SNamingTokensEditableTextBox::ValidateTokenBraces(const FText& InTokenizedText, FText& OutError)
@@ -508,11 +551,11 @@ bool SNamingTokensEditableTextBox::ParseArgs(const FString& InTokenizedTextStrin
 
 void SNamingTokensEditableTextBox::OnEditableTextChanged(const FText& InTokenizedText)
 {
+	OnTokenizedTextChanged.ExecuteIfBound(InTokenizedText);
+
 	ValidateTokenizedText(InTokenizedText);
 
 	UpdateFilters();
-	
-	OnTokenizedTextChanged.ExecuteIfBound(InTokenizedText);
 }
 
 void SNamingTokensEditableTextBox::OnEditableTextCommitted(const FText& InTokenizedText, ETextCommit::Type InCommitType)
@@ -527,7 +570,10 @@ void SNamingTokensEditableTextBox::OnEditableTextCommitted(const FText& InTokeni
 		}
 	}
 	
-	ValidateTokenizedText(InTokenizedText);
+	if (!ValidateTokenizedText(InTokenizedText))
+	{
+		return;
+	}
 
 	if (!TokenizedText.IsBound())
 	{
@@ -537,6 +583,8 @@ void SNamingTokensEditableTextBox::OnEditableTextCommitted(const FText& InTokeni
 	OnTokenizedTextCommitted.ExecuteIfBound(InTokenizedText, InCommitType);
 
 	EvaluateNamingTokens();
+
+	ExitEditingMode();
 }
 
 void SNamingTokensEditableTextBox::UpdateFilters() const
diff --git a/Engine/Plugins/Developer/NamingTokens/Source/NamingTokensUI/Public/SNamingTokensEditableTextBox.h b/Engine/Plugins/Developer/NamingTokens/Source/NamingTokensUI/Public/SNamingTokensEditableTextBox.h
--- a/Engine/Plugins/Developer/NamingTokens/Source/NamingTokensUI/Public/SNamingTokensEditableTextBox.h
+++ b/Engine/Plugins/Developer/NamingTokens/Source/NamingTokensUI/Public/SNamingTokensEditableTextBox.h
@@ -38,6 +38,8 @@ public:
 		, _DisplayTokenIcon(true)
 		, _DisplayErrorMessage(true)
 		, _DisplayBorderImage(true)
+		, _TextBoxPadding()
+		, _SupportInlineEdit(false)
 		, _EnableSuggestionDropdown(true)
 		, _FullyQualifyGlobalTokenSuggestions(false)
 	{}
@@ -83,9 +85,15 @@ public:
 		/** If we should use the border image. */
 		SLATE_ATTRIBUTE(bool, DisplayBorderImage)
 
-		/** Allows custom validation. */
+		/** Padding for the editable textbox (overrides the textbox style)*/
+		SLATE_ATTRIBUTE(FMargin, TextBoxPadding)
+
+		/** Allows custom validation before committing the tokenized text. */
 		SLATE_EVENT(FOnVerifyTextChanged, OnValidateTokenizedText)
 
+		/** Allows custom validation after the text has been resolved. */
+		SLATE_EVENT(FOnVerifyTextChanged, OnValidateResolvedText)
+
 		/** Called whenever the (tokenized) text is changed interactively by the user. */
 		SLATE_EVENT(FOnTextChanged, OnTextChanged)
 		
@@ -104,6 +112,12 @@ public:
 		/** Contexts to provide to evaluation. */
 		SLATE_ARGUMENT(TArray<TObjectPtr<UObject>>, Contexts)
 
+		/** 
+		 * If true, the look and feel of the editable textbox will be treated like an inline editable textbox.
+		 * This overrides the DisplayTokenIcon and DisplayBorderImage attributes, as well as the visibility of the child editable textbox widget
+		 */
+		SLATE_ARGUMENT(bool, SupportInlineEdit)
+
 		/** If drop down suggestions should be enabled. */
 		SLATE_ARGUMENT(bool, EnableSuggestionDropdown)
 		
@@ -136,6 +150,9 @@ public:
 	/** Sets the argument style to use. */
 	NAMINGTOKENSUI_API void SetArgumentStyle(const FTextBlockStyle* InArgumentStyle);
 
+	/** If the textbox should look and behave like an inline editable textbox */
+	NAMINGTOKENSUI_API void SetSupportInlineEdit(bool bInSupportInlineEdit);
+
 	/** If the suggestion box should be enabled. */
 	NAMINGTOKENSUI_API void SetEnableSuggestionDropdown(bool bInEnableSuggestionDropdown);
 	
@@ -148,6 +165,12 @@ public:
 	/** Evaluate the token text and update the resolved text. */
 	NAMINGTOKENSUI_API void EvaluateNamingTokens();
 	
+	/** If inline editing is supported, gives focus to the editable text, and sets the visibility and display options to resemble an editable textbox */
+	NAMINGTOKENSUI_API void EnterEditingMode();
+
+	/** If inline editing is supported, sets the visibility and display options to resemble a read-only textblock */
+	NAMINGTOKENSUI_API void ExitEditingMode();
+
 private:
 	/** @return Text contents based on current state */
 	FText GetText() const;
@@ -159,6 +182,7 @@ private:
 
 	void SetError(const FText& InError) const;
 
+	NAMINGTOKENSUI_API virtual bool SupportsKeyboardFocus() const override;
 	NAMINGTOKENSUI_API virtual bool HasKeyboardFocus() const override;
 	NAMINGTOKENSUI_API virtual FReply OnFocusReceived(const FGeometry& MyGeometry, const FFocusEvent& InFocusEvent) override;
 	NAMINGTOKENSUI_API virtual void OnFocusChanging(const FWeakWidgetPath& PreviousFocusPath, const FWidgetPath& NewWidgetPath, const FFocusEvent& InFocusEvent) override;
@@ -280,6 +304,9 @@ private:
 	/** If the border image is used. */
 	TAttribute<bool> DisplayBorderImage;
 	
+	/** Textbox padding (overrides style) */
+	TAttribute<FMargin> TextBoxPaddingOverride;
+
 	/** Attribute checking if we can display the resolved text. */
 	TAttribute<bool> CanDisplayResolvedText;
 	TAttribute<TConstArrayView<FString>> ValidArguments;
@@ -287,6 +314,9 @@ private:
 	/** Callback to verify tokenized text when changed. Will return an error message to denote problems. */
     FOnVerifyTextChanged OnValidateTokenizedText;
 
+	/** Callback to verify resolved text when changed. Will return an error message to denote problems. */
+	FOnVerifyTextChanged OnValidateResolvedText;
+
 	/** Callback when tokenized text is changed. */
 	FOnTextChanged OnTokenizedTextChanged;
 
@@ -302,6 +332,9 @@ private:
 	/** If global tokens should have their namespace inserted too from suggestions. */
 	bool bFullyQualifyGlobalTokenSuggestions = false;
 
+	/** If the widget supports inline editing */
+	bool bSupportInlineEdit = true;
+
 	/** If suggestions are enabled by default. */
 	bool bEnableSuggestionDropdown = true;
 	
