NamingTokens: Fix naming tokens editable textbox widget to properly set visibility when in inline edit mode.
Also fixes the error reporting to clear the error when the text is committed, since the invalid text is reverted in this case.
#rnx

diff --git a/Engine/Plugins/Developer/NamingTokens/Source/NamingTokensUI/Private/SNamingTokensEditableTextBox.cpp b/Engine/Plugins/Developer/NamingTokens/Source/NamingTokensUI/Private/SNamingTokensEditableTextBox.cpp
--- a/Engine/Plugins/Developer/NamingTokens/Source/NamingTokensUI/Private/SNamingTokensEditableTextBox.cpp
+++ b/Engine/Plugins/Developer/NamingTokens/Source/NamingTokensUI/Private/SNamingTokensEditableTextBox.cpp
@@ -200,6 +200,11 @@ void SNamingTokensEditableTextBox::Construct(const FArguments& InArgs)
 		]
 	];
 	
+	if (bSupportInlineEdit)
+	{
+		SetVisibility(EVisibility::HitTestInvisible);
+	}
+
 	OnEditableTextCommitted(TokenizedText.Get(), ETextCommit::Default);
 }
 
@@ -423,7 +428,7 @@ void SNamingTokensEditableTextBox::EnterEditingMode()
 		FSlateApplication::Get().SetKeyboardFocus(EditableText.ToSharedRef());
 		EditableText->SelectAllText();
 
-		EditableText->SetVisibility(EVisibility::Visible);
+		SetVisibility(EVisibility::Visible);
 		DisplayTokenIcon = true;
 		DisplayBorderImage = true;
 	}
@@ -434,7 +439,7 @@ void SNamingTokensEditableTextBox::ExitEditingMode()
 	if (bSupportInlineEdit)
 	{
 		// Update the visibility to more closely match that of a normal read-only text block
-		EditableText->SetVisibility(EVisibility::HitTestInvisible);
+		SetVisibility(EVisibility::HitTestInvisible);
 		DisplayTokenIcon = false;
 		DisplayBorderImage = false;
 	}
@@ -569,9 +574,13 @@ void SNamingTokensEditableTextBox::OnEditableTextCommitted(const FText& InTokeni
 			CloseSuggestionDropdown();
 		}
 	}
-	
+
+	ExitEditingMode();
+
 	if (!ValidateTokenizedText(InTokenizedText))
 	{
+		// The invalid text will not be committed, so clear the error message
+		SetError(FText::GetEmpty());
 		return;
 	}
 
@@ -583,13 +592,11 @@ void SNamingTokensEditableTextBox::OnEditableTextCommitted(const FText& InTokeni
 	OnTokenizedTextCommitted.ExecuteIfBound(InTokenizedText, InCommitType);
 
 	EvaluateNamingTokens();
-
-	ExitEditingMode();
 }
 
 void SNamingTokensEditableTextBox::UpdateFilters() const
 {
-	if (EditableText.IsValid() && FilterWidget.IsValid() && DropdownAnchor->IsOpen())
+	if (EditableText.IsValid() && FilterWidget.IsValid() && DropdownAnchor.IsValid() && DropdownAnchor->IsOpen())
 	{
 		// Find the text we're typing in between brackets and pass that as a filter to our suggestion drop down.
 
diff --git a/Engine/Plugins/Developer/NamingTokens/Source/NamingTokensUI/Public/SNamingTokensEditableTextBox.h b/Engine/Plugins/Developer/NamingTokens/Source/NamingTokensUI/Public/SNamingTokensEditableTextBox.h
--- a/Engine/Plugins/Developer/NamingTokens/Source/NamingTokensUI/Public/SNamingTokensEditableTextBox.h
+++ b/Engine/Plugins/Developer/NamingTokens/Source/NamingTokensUI/Public/SNamingTokensEditableTextBox.h
@@ -114,7 +114,7 @@ public:
 
 		/** 
 		 * If true, the look and feel of the editable textbox will be treated like an inline editable textbox.
-		 * This overrides the DisplayTokenIcon and DisplayBorderImage attributes, as well as the visibility of the child editable textbox widget
+		 * This overrides the DisplayTokenIcon and DisplayBorderImage attributes, as well as the visibility
 		 */
 		SLATE_ARGUMENT(bool, SupportInlineEdit)
 
@@ -333,7 +333,7 @@ private:
 	bool bFullyQualifyGlobalTokenSuggestions = false;
 
 	/** If the widget supports inline editing */
-	bool bSupportInlineEdit = true;
+	bool bSupportInlineEdit = false;
 
 	/** If suggestions are enabled by default. */
 	bool bEnableSuggestionDropdown = true;
