NamingTokens: Add a new global token, {source], which is a special token used by asset tools for naming assets from other assets.
#rnx

diff --git a/Engine/Plugins/Developer/NamingTokens/Source/NamingTokens/Private/GlobalNamingTokens.cpp b/Engine/Plugins/Developer/NamingTokens/Source/NamingTokens/Private/GlobalNamingTokens.cpp
--- a/Engine/Plugins/Developer/NamingTokens/Source/NamingTokens/Private/GlobalNamingTokens.cpp
+++ b/Engine/Plugins/Developer/NamingTokens/Source/NamingTokens/Private/GlobalNamingTokens.cpp
@@ -33,6 +33,15 @@ void UGlobalNamingTokens::OnCreateDefaultTokens(TArray<FNamingTokenData>& Tokens
 		})
 	});
 
+	// Returns empty text intentionally, as this token is used by Asset Tools to substitute in the name of a source object
+	Tokens.Add({
+		TEXT("source"),
+		LOCTEXT("Source", "Source"),
+		FNamingTokenData::FTokenProcessorDelegateNative::CreateLambda([]() {
+			return FText::GetEmpty();
+		})
+	});
+
 	Tokens.Add({
 	    TEXT("yyyy"),
 	    LOCTEXT("Year4Digit", "Year (4 digit)"),
