NamingTokens: Add private namespaces. When a namepace is made private it won't evaluate through the subsystem or show up in most UI by default. It must be explicitly added to the AdditionalNamespacesToInclude filter to be used.

Namespaces added to AdditionalNamespacesToInclude UMG widget will no longer autocomplete the namespace by default.
Add new bAutoCompleteAdditionalNamespacesToInclude setting to widget to allow that behavior.

Improve UMG category prioritization.

diff --git a/Engine/Plugins/Developer/NamingTokens/Source/NamingTokens/Private/NamingTokens.cpp b/Engine/Plugins/Developer/NamingTokens/Source/NamingTokens/Private/NamingTokens.cpp
--- a/Engine/Plugins/Developer/NamingTokens/Source/NamingTokens/Private/NamingTokens.cpp
+++ b/Engine/Plugins/Developer/NamingTokens/Source/NamingTokens/Private/NamingTokens.cpp
@@ -307,6 +307,11 @@ TArray<FNamingTokenData> UNamingTokens::GetAllTokens() const
 	return AllTokens;
 }
 
+bool UNamingTokens::IsPrivateNamespace() const
+{
+	return bIsPrivateNamespace;
+}
+
 FDateTime UNamingTokens::GetCurrentDateTime_Implementation() const
 {
 	return CurrentEvaluationData.CurrentDateTime;
diff --git a/Engine/Plugins/Developer/NamingTokens/Source/NamingTokens/Private/NamingTokensEngineSubsystem.cpp b/Engine/Plugins/Developer/NamingTokens/Source/NamingTokens/Private/NamingTokensEngineSubsystem.cpp
--- a/Engine/Plugins/Developer/NamingTokens/Source/NamingTokens/Private/NamingTokensEngineSubsystem.cpp
+++ b/Engine/Plugins/Developer/NamingTokens/Source/NamingTokens/Private/NamingTokensEngineSubsystem.cpp
@@ -16,7 +16,7 @@
 
 #include UE_INLINE_GENERATED_CPP_BY_NAME(NamingTokensEngineSubsystem)
 
-TArray<FString> UNamingTokensEngineSubsystem::GetAllNamespaces() const
+TArray<FString> UNamingTokensEngineSubsystem::GetAllNamespaces(bool bIncludePrivateNamespaces) const
 {
 	const FAssetRegistryModule& AssetRegistryModule = FModuleManager::LoadModuleChecked<FAssetRegistryModule>("AssetRegistry");
 	const IAssetRegistry& AssetRegistry = AssetRegistryModule.Get();
@@ -29,6 +29,14 @@ TArray<FString> UNamingTokensEngineSubsystem::GetAllNamespaces() const
 	TSet<FString> DiscoveredNamespaces;
 	for (const FAssetData& NamingTokenAssetData : NamingTokenAssets)
 	{
+		bool bIsPrivateNamespace = false;
+		NamingTokenAssetData.GetTagValue(UNamingTokens::GetIsPrivateNamespacePropertyName(), bIsPrivateNamespace);
+
+		if (bIsPrivateNamespace && !bIncludePrivateNamespaces)
+		{
+			continue;
+		}
+		
 		FString FoundNamespace;
 		NamingTokenAssetData.GetTagValue(UNamingTokens::GetNamespacePropertyName(), FoundNamespace);
 
@@ -43,8 +51,14 @@ TArray<FString> UNamingTokensEngineSubsystem::GetAllNamespaces() const
 		constexpr EClassFlags InvalidClassFlags = CLASS_Deprecated | CLASS_NewerVersionExists | CLASS_Abstract;
 		if (DerivedClass->IsNative() && !DerivedClass->HasAnyClassFlags(InvalidClassFlags))
 		{
-			FString Namepsace = DerivedClass->GetDefaultObject<UNamingTokens>()->GetNamespace();
-			DiscoveredNamespaces.Add(Namepsace);
+			const UNamingTokens* NamingTokens = CastChecked<UNamingTokens>(DerivedClass->GetDefaultObject());
+			if (NamingTokens->IsPrivateNamespace() && !bIncludePrivateNamespaces)
+			{
+				continue;
+			}
+			
+			FString Namespace = NamingTokens->GetNamespace();
+			DiscoveredNamespaces.Add(Namespace);
 		}
 	}
 
@@ -238,6 +252,11 @@ FNamingTokenResultData UNamingTokensEngineSubsystem::EvaluateTokenText(const FTe
 
 		if (UNamingTokens* NamingTokens = GetNamingTokensFunction.Execute(Namespace))
 		{
+			if (NamingTokens->IsPrivateNamespace() && !InFilter.AdditionalNamespacesToInclude.Contains(Namespace))
+			{
+				continue;
+			}
+			
 			Result = NamingTokens->EvaluateTokenText(ProcessedTokenText, EvaluationData);
 
 			// Properly record individual token evaluation data. The order should be consistent with the order
diff --git a/Engine/Plugins/Developer/NamingTokens/Source/NamingTokens/Public/NamingTokens.h b/Engine/Plugins/Developer/NamingTokens/Source/NamingTokens/Public/NamingTokens.h
--- a/Engine/Plugins/Developer/NamingTokens/Source/NamingTokens/Public/NamingTokens.h
+++ b/Engine/Plugins/Developer/NamingTokens/Source/NamingTokens/Public/NamingTokens.h
@@ -96,6 +96,8 @@ public:
 	const FString& GetNamespace() const { return Namespace; }
 	/** The name of the namespace property. */
 	static FName GetNamespacePropertyName() { return GET_MEMBER_NAME_CHECKED(UNamingTokens, Namespace); }
+	/** The name of IsPrivateNamespace property. */
+	static FName GetIsPrivateNamespacePropertyName() { return GET_MEMBER_NAME_CHECKED(UNamingTokens, bIsPrivateNamespace); }
 
 	/** Retrieve the namespace's friendly display name. */
 	const FText& GetNamespaceDisplayName() const { return NamespaceDisplayName; }
@@ -110,6 +112,10 @@ public:
 	/** Retrieve the current datetime. By default, this uses shared data so results are consistent across runs. */
 	UFUNCTION(BlueprintCallable, BlueprintNativeEvent, Category = "NamingTokens")
 	UE_API FDateTime GetCurrentDateTime() const;
+
+	/** If this namespace is private. */
+	UFUNCTION(BlueprintCallable, Category = "NamingTokens")
+	UE_API bool IsPrivateNamespace() const;
 	
 protected:
 	UE_API virtual FDateTime GetCurrentDateTime_Implementation() const;
@@ -192,6 +198,13 @@ protected:
 	 */
 	UPROPERTY(EditDefaultsOnly, AssetRegistrySearchable, Category = "NamingTokens")
 	FText NamespaceDisplayName;
+
+	/**
+	 * If the namespace is considered private. This prevents this namespace from being gathered for display
+	 * or processing unless explicitly specified.
+	 */
+	UPROPERTY(AdvancedDisplay, EditDefaultsOnly, AssetRegistrySearchable, Category = "NamingTokens")
+	bool bIsPrivateNamespace = false;
 	
 #if WITH_EDITORONLY_DATA
 private:
diff --git a/Engine/Plugins/Developer/NamingTokens/Source/NamingTokens/Public/NamingTokensEngineSubsystem.h b/Engine/Plugins/Developer/NamingTokens/Source/NamingTokens/Public/NamingTokensEngineSubsystem.h
--- a/Engine/Plugins/Developer/NamingTokens/Source/NamingTokens/Public/NamingTokensEngineSubsystem.h
+++ b/Engine/Plugins/Developer/NamingTokens/Source/NamingTokens/Public/NamingTokensEngineSubsystem.h
@@ -144,7 +144,7 @@ public:
 
 	/** Retrieve all discovered namespaces. */
 	UFUNCTION(BlueprintCallable, Category = "NamingTokens")
-	UE_API TArray<FString> GetAllNamespaces() const;
+	UE_API TArray<FString> GetAllNamespaces(bool bIncludePrivateNamespaces = false) const;
 
 	/** Register a delegate to filter the list of namespaces that can be used to evaluate tokens */
 	UE_API void RegisterNamespaceFilter(const FName OwnerName, FFilterNamespace Delegate);
diff --git a/Engine/Plugins/Developer/NamingTokens/Source/NamingTokensUI/Private/NamingTokensEditableText.cpp b/Engine/Plugins/Developer/NamingTokens/Source/NamingTokensUI/Private/NamingTokensEditableText.cpp
--- a/Engine/Plugins/Developer/NamingTokens/Source/NamingTokensUI/Private/NamingTokensEditableText.cpp
+++ b/Engine/Plugins/Developer/NamingTokens/Source/NamingTokensUI/Private/NamingTokensEditableText.cpp
@@ -160,6 +160,7 @@ void UNamingTokensEditableText::SynchronizeProperties()
 	MyNamingTokensEditableText->SetNormalStyle(&*EditableTextBoxStyle);
 	MyNamingTokensEditableText->SetArgumentStyle(&ArgumentStyle);
 	MyNamingTokensEditableText->SetEnableSuggestionDropdown(bEnableSuggestionDropdown);
+	MyNamingTokensEditableText->SetFullyQualifyFilteredNamespaces(bFullyQualifyFilteredNamespaces);
 	
 	MyMultiLineEditableText->SetText(MyNamingTokensEditableText->GetBoundText());
 	MyMultiLineEditableText->SetHintText(HintTextBinding);
diff --git a/Engine/Plugins/Developer/NamingTokens/Source/NamingTokensUI/Private/SNamingTokensDataTreeView.cpp b/Engine/Plugins/Developer/NamingTokens/Source/NamingTokensUI/Private/SNamingTokensDataTreeView.cpp
--- a/Engine/Plugins/Developer/NamingTokens/Source/NamingTokensUI/Private/SNamingTokensDataTreeView.cpp
+++ b/Engine/Plugins/Developer/NamingTokens/Source/NamingTokensUI/Private/SNamingTokensDataTreeView.cpp
@@ -148,7 +148,7 @@ void SNamingTokenDataTreeViewWidget::Construct(const FArguments& InArgs)
 		HeaderRow->AddColumn(ColumnArgs);
 	}
 
-	PopulateTreeItems();
+	PopulateTreeItems(InArgs._FilterArgs);
 
 	ChildSlot
 	[
@@ -171,14 +171,14 @@ void SNamingTokenDataTreeViewWidget::Construct(const FArguments& InArgs)
 	FilterTreeItems(FString());
 }
 
-void SNamingTokenDataTreeViewWidget::PopulateTreeItems()
+void SNamingTokenDataTreeViewWidget::PopulateTreeItems(const FNamingTokenFilterArgs& InFilterArgs)
 {
 	RootTreeItems.Reset();
 	if (GEngine)
 	{
 		if (const UNamingTokensEngineSubsystem* NamingTokensSubsystem = GEngine->GetEngineSubsystem<UNamingTokensEngineSubsystem>())
 		{
-			const TArray<FString> Namespaces = NamingTokensSubsystem->GetAllNamespaces();
+			const TArray<FString> Namespaces = NamingTokensSubsystem->GetAllNamespaces(/*bIncludePrivateNamespaces*/ true);
 			RootTreeItems.Reserve(Namespaces.Num());
 
 			int32 TotalCount = 0;
@@ -192,6 +192,13 @@ void SNamingTokenDataTreeViewWidget::PopulateTreeItems()
 				
 				if (const UNamingTokens* NamingTokens = NamingTokensSubsystem->GetNamingTokens(Namespace))
 				{
+					const bool bIsPrivate = NamingTokens->IsPrivateNamespace();
+					// Filter out any private namespaces unless specified.
+					if (bIsPrivate && !InFilterArgs.AdditionalNamespacesToInclude.Contains(Namespace))
+					{
+						continue;
+					}
+
 					NamespaceTreeItem->NamespaceDisplayName = NamingTokens->GetNamespaceDisplayName();
 					
 					const TArray<FNamingTokenData> Tokens = NamingTokens->GetAllTokens();
diff --git a/Engine/Plugins/Developer/NamingTokens/Source/NamingTokensUI/Private/SNamingTokensDataTreeView.h b/Engine/Plugins/Developer/NamingTokens/Source/NamingTokensUI/Private/SNamingTokensDataTreeView.h
--- a/Engine/Plugins/Developer/NamingTokens/Source/NamingTokensUI/Private/SNamingTokensDataTreeView.h
+++ b/Engine/Plugins/Developer/NamingTokens/Source/NamingTokensUI/Private/SNamingTokensDataTreeView.h
@@ -3,6 +3,7 @@
 #pragma once
 
 #include "NamingTokenData.h"
+#include "NamingTokensEngineSubsystem.h"
 #include "Utils/NamingTokenUtils.h"
 
 #include "Widgets/Views/STreeView.h"
@@ -171,6 +172,8 @@ public:
 		SLATE_EVENT(FOnItemSelected, OnItemSelected)
 		/** Event if this widget or its children receives focus. */
 		SLATE_EVENT(FSimpleDelegate, OnFocused)
+		/** Additional filter args to apply. */
+		SLATE_ARGUMENT(FNamingTokenFilterArgs, FilterArgs)
 	SLATE_END_ARGS()
 
 	virtual ~SNamingTokenDataTreeViewWidget() override;
@@ -178,7 +181,7 @@ public:
 	void Construct(const FArguments& InArgs);
 
 	/** Populate all possible tree items */
-	void PopulateTreeItems();
+	void PopulateTreeItems(const FNamingTokenFilterArgs& InFilterArgs);
 
 	/** Filter down to specific tokens and namespaces. */
 	void FilterTreeItems(const FString& InFilterText);
diff --git a/Engine/Plugins/Developer/NamingTokens/Source/NamingTokensUI/Private/SNamingTokensEditableTextBox.cpp b/Engine/Plugins/Developer/NamingTokens/Source/NamingTokensUI/Private/SNamingTokensEditableTextBox.cpp
--- a/Engine/Plugins/Developer/NamingTokens/Source/NamingTokensUI/Private/SNamingTokensEditableTextBox.cpp
+++ b/Engine/Plugins/Developer/NamingTokens/Source/NamingTokensUI/Private/SNamingTokensEditableTextBox.cpp
@@ -72,6 +72,7 @@ void SNamingTokensEditableTextBox::Construct(const FArguments& InArgs)
 	bSupportInlineEdit = InArgs._SupportInlineEdit;
 	bEnableSuggestionDropdown = InArgs._EnableSuggestionDropdown;
 	bFullyQualifyGlobalTokenSuggestions = InArgs._FullyQualifyGlobalTokenSuggestions;
+	bFullyQualifyFilteredNamespaces = InArgs._FullyQualifyFilteredNamespaces;
 	
 	SetFilterArgs(InArgs._FilterArgs);
 	SetContexts(InArgs._Contexts);
@@ -276,6 +277,11 @@ void SNamingTokensEditableTextBox::SetFilterArgs(const FNamingTokenFilterArgs& I
 	NamingTokenFilterArgs = InFilterArgs;
 }
 
+void SNamingTokensEditableTextBox::SetFullyQualifyFilteredNamespaces(bool bNewValue)
+{
+	bFullyQualifyFilteredNamespaces = bNewValue;
+}
+
 void SNamingTokensEditableTextBox::SetContexts(const TArray<UObject*>& InContexts)
 {
 	NamingTokenContexts = InContexts;
@@ -740,7 +746,8 @@ TSharedRef<SWidget> SNamingTokensEditableTextBox::OnGetDropdownContent()
 	}
 	TSharedRef<SNamingTokenDataTreeViewWidget> Widget = SAssignNew(FilterWidget, SNamingTokenDataTreeViewWidget)
 		.OnItemSelected(this, &SNamingTokensEditableTextBox::OnSuggestionDoubleClicked)
-		.OnFocused(this, &SNamingTokensEditableTextBox::OnSuggestionDropdownReceivedFocus);
+		.OnFocused(this, &SNamingTokensEditableTextBox::OnSuggestionDropdownReceivedFocus)
+		.FilterArgs(NamingTokenFilterArgs);
 	return Widget;
 }
 
@@ -829,7 +836,11 @@ void SNamingTokensEditableTextBox::OnSuggestionSelected(TSharedPtr<FNamingTokenD
 		if (const TSharedPtr<FCursorTokenData> TokenLocation = GetCurrentCursorTokenData())
 		{
 			// Either insert just the token, or the namespace:token.
-			const bool bIsFullyQualified = !SelectedItem->bIsGlobal || bFullyQualifyGlobalTokenSuggestions;
+			bool bIsFullyQualified = !SelectedItem->bIsGlobal || bFullyQualifyGlobalTokenSuggestions;
+			if (!bFullyQualifyFilteredNamespaces && NamingTokenFilterArgs.AdditionalNamespacesToInclude.Contains(SelectedItem->Namespace))
+			{
+				bIsFullyQualified = false;
+			}
 		
 			FString TextToInsert = SelectedItem->ToString(bIsFullyQualified);
 			if (SelectedItem->IsNamespace())
diff --git a/Engine/Plugins/Developer/NamingTokens/Source/NamingTokensUI/Public/NamingTokensEditableText.h b/Engine/Plugins/Developer/NamingTokens/Source/NamingTokensUI/Public/NamingTokensEditableText.h
--- a/Engine/Plugins/Developer/NamingTokens/Source/NamingTokensUI/Public/NamingTokensEditableText.h
+++ b/Engine/Plugins/Developer/NamingTokens/Source/NamingTokensUI/Public/NamingTokensEditableText.h
@@ -13,7 +13,7 @@ struct FEditableTextBoxStyle;
 /**
  * Editable text for displaying tokenized strings in either their unevaluated or resolved form.
  */
-UCLASS(MinimalAPI, Config = NamingTokens, DefaultConfig, ClassGroup = UI, meta = (Category = "NamingTokens UI", DisplayName = "Naming Tokens Editable Text Box", PrioritizeCategories = "Content"))
+UCLASS(MinimalAPI, Config = NamingTokens, DefaultConfig, ClassGroup = UI, PrioritizeCategories = ("NamingTokens", "Content", "Appearance", "Behavior"), meta = (Category = "NamingTokens UI", DisplayName = "Naming Tokens Editable Text Box"))
 class UNamingTokensEditableText : public UMultiLineEditableText
 {
 	GENERATED_BODY()
@@ -104,6 +104,10 @@ protected:
 	/** If the dropdown suggestion box should be enabled. */
 	UPROPERTY(EditAnywhere, AdvancedDisplay, Category=Behavior)
 	bool bEnableSuggestionDropdown = true;
+
+	/** If namespaces included in the FilterArgs.AdditionalNamespacesToInclude should autocomplete. If false, only the token is inserted. */
+	UPROPERTY(EditAnywhere, AdvancedDisplay, Category=Behavior, meta = (EditCondition="bEnableSuggestionDropdown", EditConditionHides))
+	bool bFullyQualifyFilteredNamespaces = false;
 	
 	/** If this text box is configured for multiline. */
 	UPROPERTY(EditAnywhere, BlueprintReadWrite, Category=Appearance)
diff --git a/Engine/Plugins/Developer/NamingTokens/Source/NamingTokensUI/Public/SNamingTokensEditableTextBox.h b/Engine/Plugins/Developer/NamingTokens/Source/NamingTokensUI/Public/SNamingTokensEditableTextBox.h
--- a/Engine/Plugins/Developer/NamingTokens/Source/NamingTokensUI/Public/SNamingTokensEditableTextBox.h
+++ b/Engine/Plugins/Developer/NamingTokens/Source/NamingTokensUI/Public/SNamingTokensEditableTextBox.h
@@ -42,6 +42,7 @@ public:
 		, _SupportInlineEdit(false)
 		, _EnableSuggestionDropdown(true)
 		, _FullyQualifyGlobalTokenSuggestions(false)
+		, _FullyQualifyFilteredNamespaces(false)
 	{}
 		/** The styling of the textbox. */
 		SLATE_STYLE_ARGUMENT(FEditableTextBoxStyle, Style)
@@ -124,6 +125,9 @@ public:
 		/** If the namespace should be inserted when selecting a global suggestion. */
 		SLATE_ARGUMENT(bool, FullyQualifyGlobalTokenSuggestions)
 
+		/** If namespaces included in FilterArgs.AdditionalNamespacesToInclude should be autocompleted. */
+		SLATE_ARGUMENT(bool, FullyQualifyFilteredNamespaces)
+
 	SLATE_END_ARGS()
 
 	/** Constructs this widget with InArgs */
@@ -159,6 +163,9 @@ public:
 	/** Provide filter args to use during evaluation. */
 	NAMINGTOKENSUI_API void SetFilterArgs(const FNamingTokenFilterArgs& InFilterArgs);
 
+	/** Set whether namespaces under FilterArgs.AdditionalNamespacesToInclude should autocomplete. */
+	NAMINGTOKENSUI_API void SetFullyQualifyFilteredNamespaces(bool bNewValue);
+	
 	/** Provide contexts to naming token evaluation. */
 	NAMINGTOKENSUI_API void SetContexts(const TArray<UObject*>& InContexts);
 
@@ -332,6 +339,9 @@ private:
 	/** If global tokens should have their namespace inserted too from suggestions. */
 	bool bFullyQualifyGlobalTokenSuggestions = false;
 
+	/** If namespaces under FilterArgs.AdditionalNamespacesToInclude should autocomplete. */
+	bool bFullyQualifyFilteredNamespaces = false;
+	
 	/** If the widget supports inline editing */
 	bool bSupportInlineEdit = false;
 
