NamingTokens: Add NamespaceSuggestionPriority to widgets which will sort the dropdown suggestions based on the list order.

diff --git a/Engine/Plugins/Developer/NamingTokens/Source/NamingTokensUI/Private/NamingTokensEditableText.cpp b/Engine/Plugins/Developer/NamingTokens/Source/NamingTokensUI/Private/NamingTokensEditableText.cpp
--- a/Engine/Plugins/Developer/NamingTokens/Source/NamingTokensUI/Private/NamingTokensEditableText.cpp
+++ b/Engine/Plugins/Developer/NamingTokens/Source/NamingTokensUI/Private/NamingTokensEditableText.cpp
@@ -132,6 +132,7 @@ TSharedRef<SWidget> UNamingTokensEditableText::RebuildWidget()
 		.FilterArgs(FilterArgs)
 		.Contexts(Contexts)
 		.EnableSuggestionDropdown(bEnableSuggestionDropdown)
+		.NamespaceSuggestionPriority(NamespaceSuggestionPriority)
 		.DisplayTokenIcon_UObject(this, &UNamingTokensEditableText::GetDisplayTokenIcon)
 		.DisplayErrorMessage_UObject(this, &UNamingTokensEditableText::GetDisplayErrorMessage)
 		.DisplayBorderImage_UObject(this, &UNamingTokensEditableText::GetDisplayBorderImage)
@@ -161,6 +162,7 @@ void UNamingTokensEditableText::SynchronizeProperties()
 	MyNamingTokensEditableText->SetArgumentStyle(&ArgumentStyle);
 	MyNamingTokensEditableText->SetEnableSuggestionDropdown(bEnableSuggestionDropdown);
 	MyNamingTokensEditableText->SetFullyQualifyFilteredNamespaces(bFullyQualifyFilteredNamespaces);
+	MyNamingTokensEditableText->SetNamespaceSuggestionPriority(NamespaceSuggestionPriority);
 	
 	MyMultiLineEditableText->SetText(MyNamingTokensEditableText->GetBoundText());
 	MyMultiLineEditableText->SetHintText(HintTextBinding);
diff --git a/Engine/Plugins/Developer/NamingTokens/Source/NamingTokensUI/Private/SNamingTokensDataTreeView.cpp b/Engine/Plugins/Developer/NamingTokens/Source/NamingTokensUI/Private/SNamingTokensDataTreeView.cpp
--- a/Engine/Plugins/Developer/NamingTokens/Source/NamingTokensUI/Private/SNamingTokensDataTreeView.cpp
+++ b/Engine/Plugins/Developer/NamingTokens/Source/NamingTokensUI/Private/SNamingTokensDataTreeView.cpp
@@ -133,6 +133,7 @@ void SNamingTokenDataTreeViewWidget::Construct(const FArguments& InArgs)
 	ItemSelectionChangedDelegate = InArgs._OnSelectionChanged;
 	ItemSelectedDelegate = InArgs._OnItemSelected;
 	OnFocusedDelegate = InArgs._OnFocused;
+	NamespacePriority = InArgs._NamespacePriority;
 	
 	TSharedPtr<SHeaderRow> HeaderRow;
 	// Setup the columns.
@@ -178,9 +179,41 @@ void SNamingTokenDataTreeViewWidget::PopulateTreeItems(const FNamingTokenFilterA
 	{
 		if (const UNamingTokensEngineSubsystem* NamingTokensSubsystem = GEngine->GetEngineSubsystem<UNamingTokensEngineSubsystem>())
 		{
-			const TArray<FString> Namespaces = NamingTokensSubsystem->GetAllNamespaces(/*bIncludePrivateNamespaces*/ true);
+			TArray<FString> Namespaces = NamingTokensSubsystem->GetAllNamespaces(/*bIncludePrivateNamespaces*/ true);
 			RootTreeItems.Reserve(Namespaces.Num());
 
+			// Sort namespaces so any matching the priority list are displayed first, in that order.
+			// Any other namespaces preserve their order. It's better to do this before building out final lists
+			// so we can record the proper index order for navigation next.
+			if (!NamespacePriority.IsEmpty())
+			{
+				TMap<FString, int32> NamespaceOrder;
+				NamespaceOrder.Reserve(NamespacePriority.Num());
+				for (int32 Idx = 0; Idx < NamespacePriority.Num(); ++Idx)
+				{
+					NamespaceOrder.Add(NamespacePriority[Idx], Idx);
+				}
+
+				Namespaces.StableSort([&](const FString& A, const FString& B)
+				{
+					const int32* IndexA = NamespaceOrder.Find(A);
+					const int32* IndexB = NamespaceOrder.Find(B);
+
+					if (IndexA && IndexB)
+					{
+						return *IndexA < *IndexB;
+					}
+
+					if (IndexA)
+					{
+						return true;
+					}
+
+					// Keep the original order.
+					return false;
+				});
+			}
+			
 			int32 TotalCount = 0;
 			
 			for (const FString& Namespace : Namespaces)
diff --git a/Engine/Plugins/Developer/NamingTokens/Source/NamingTokensUI/Private/SNamingTokensDataTreeView.h b/Engine/Plugins/Developer/NamingTokens/Source/NamingTokensUI/Private/SNamingTokensDataTreeView.h
--- a/Engine/Plugins/Developer/NamingTokens/Source/NamingTokensUI/Private/SNamingTokensDataTreeView.h
+++ b/Engine/Plugins/Developer/NamingTokens/Source/NamingTokensUI/Private/SNamingTokensDataTreeView.h
@@ -174,6 +174,8 @@ public:
 		SLATE_EVENT(FSimpleDelegate, OnFocused)
 		/** Additional filter args to apply. */
 		SLATE_ARGUMENT(FNamingTokenFilterArgs, FilterArgs)
+		/** Display priority of namespaces. */
+		SLATE_ARGUMENT(TArray<FString>, NamespacePriority)
 	SLATE_END_ARGS()
 
 	virtual ~SNamingTokenDataTreeViewWidget() override;
@@ -253,6 +255,9 @@ private:
 	/** Items filtered in a search. */
 	TArray<FNamingTokenDataTreeItemPtr> FilteredRootTreeItems;
 
+	/** Order to display namespaces. */
+	TArray<FString> NamespacePriority;
+
 	/** The last filter applied. */
 	FString LastFilterText;
 
diff --git a/Engine/Plugins/Developer/NamingTokens/Source/NamingTokensUI/Private/SNamingTokensEditableTextBox.cpp b/Engine/Plugins/Developer/NamingTokens/Source/NamingTokensUI/Private/SNamingTokensEditableTextBox.cpp
--- a/Engine/Plugins/Developer/NamingTokens/Source/NamingTokensUI/Private/SNamingTokensEditableTextBox.cpp
+++ b/Engine/Plugins/Developer/NamingTokens/Source/NamingTokensUI/Private/SNamingTokensEditableTextBox.cpp
@@ -73,6 +73,7 @@ void SNamingTokensEditableTextBox::Construct(const FArguments& InArgs)
 	bEnableSuggestionDropdown = InArgs._EnableSuggestionDropdown;
 	bFullyQualifyGlobalTokenSuggestions = InArgs._FullyQualifyGlobalTokenSuggestions;
 	bFullyQualifyFilteredNamespaces = InArgs._FullyQualifyFilteredNamespaces;
+	SetNamespaceSuggestionPriority(InArgs._NamespaceSuggestionPriority);
 	
 	SetFilterArgs(InArgs._FilterArgs);
 	SetContexts(InArgs._Contexts);
@@ -282,6 +283,11 @@ void SNamingTokensEditableTextBox::SetFullyQualifyFilteredNamespaces(bool bNewVa
 	bFullyQualifyFilteredNamespaces = bNewValue;
 }
 
+void SNamingTokensEditableTextBox::SetNamespaceSuggestionPriority(const TArray<FString>& InNamespaces)
+{
+	NamespaceSuggestionPriority = InNamespaces;
+}
+
 void SNamingTokensEditableTextBox::SetContexts(const TArray<UObject*>& InContexts)
 {
 	NamingTokenContexts = InContexts;
@@ -747,7 +753,8 @@ TSharedRef<SWidget> SNamingTokensEditableTextBox::OnGetDropdownContent()
 	TSharedRef<SNamingTokenDataTreeViewWidget> Widget = SAssignNew(FilterWidget, SNamingTokenDataTreeViewWidget)
 		.OnItemSelected(this, &SNamingTokensEditableTextBox::OnSuggestionDoubleClicked)
 		.OnFocused(this, &SNamingTokensEditableTextBox::OnSuggestionDropdownReceivedFocus)
-		.FilterArgs(NamingTokenFilterArgs);
+		.FilterArgs(NamingTokenFilterArgs)
+		.NamespacePriority(NamespaceSuggestionPriority);
 	return Widget;
 }
 
diff --git a/Engine/Plugins/Developer/NamingTokens/Source/NamingTokensUI/Public/NamingTokensEditableText.h b/Engine/Plugins/Developer/NamingTokens/Source/NamingTokensUI/Public/NamingTokensEditableText.h
--- a/Engine/Plugins/Developer/NamingTokens/Source/NamingTokensUI/Public/NamingTokensEditableText.h
+++ b/Engine/Plugins/Developer/NamingTokens/Source/NamingTokensUI/Public/NamingTokensEditableText.h
@@ -100,6 +100,10 @@ protected:
 	/** Set naming token filter args to use during token evaluation. */
 	UPROPERTY(EditAnywhere, BlueprintReadWrite, Category=NamingTokens)
 	FNamingTokenFilterArgs FilterArgs;
+
+	/** A list prioritizing namespaces to show first in the dropdown. */
+	UPROPERTY(EditAnywhere, AdvancedDisplay, Category=Behavior, meta = (EditCondition="bEnableSuggestionDropdown", EditConditionHides, DisplayAfter=bFullyQualifyFilteredNamespaces))
+	TArray<FString> NamespaceSuggestionPriority;
 	
 	/** If the dropdown suggestion box should be enabled. */
 	UPROPERTY(EditAnywhere, AdvancedDisplay, Category=Behavior)
diff --git a/Engine/Plugins/Developer/NamingTokens/Source/NamingTokensUI/Public/SNamingTokensEditableTextBox.h b/Engine/Plugins/Developer/NamingTokens/Source/NamingTokensUI/Public/SNamingTokensEditableTextBox.h
--- a/Engine/Plugins/Developer/NamingTokens/Source/NamingTokensUI/Public/SNamingTokensEditableTextBox.h
+++ b/Engine/Plugins/Developer/NamingTokens/Source/NamingTokensUI/Public/SNamingTokensEditableTextBox.h
@@ -128,6 +128,9 @@ public:
 		/** If namespaces included in FilterArgs.AdditionalNamespacesToInclude should be autocompleted. */
 		SLATE_ARGUMENT(bool, FullyQualifyFilteredNamespaces)
 
+		/** Set the namespace priority to use when sorting dropdown suggestions. */
+		SLATE_ARGUMENT(TArray<FString>, NamespaceSuggestionPriority)
+
 	SLATE_END_ARGS()
 
 	/** Constructs this widget with InArgs */
@@ -165,6 +168,9 @@ public:
 
 	/** Set whether namespaces under FilterArgs.AdditionalNamespacesToInclude should autocomplete. */
 	NAMINGTOKENSUI_API void SetFullyQualifyFilteredNamespaces(bool bNewValue);
+
+	/** Set the namespace priority to use when sorting dropdown suggestions. */
+	NAMINGTOKENSUI_API void SetNamespaceSuggestionPriority(const TArray<FString>& InNamespaces);
 	
 	/** Provide contexts to naming token evaluation. */
 	NAMINGTOKENSUI_API void SetContexts(const TArray<UObject*>& InContexts);
@@ -259,6 +265,9 @@ private:
 
 	/** Contexts to provide to evaluation. */
 	TArray<TObjectPtr<UObject>> NamingTokenContexts;
+
+	/** A list prioritizing namespaces to show first in the dropdown. */
+	TArray<FString> NamespaceSuggestionPriority;
 	
 	/** Editable text widget. */
 	TSharedPtr<SMultiLineEditableText> EditableText;
