NamingTokens: Fix the naming tokens textbox to not fire the OnTextCommitted callback during widget construction.
#rnx

diff --git a/Engine/Plugins/Developer/NamingTokens/Source/NamingTokensUI/Private/SNamingTokensEditableTextBox.cpp b/Engine/Plugins/Developer/NamingTokens/Source/NamingTokensUI/Private/SNamingTokensEditableTextBox.cpp
--- a/Engine/Plugins/Developer/NamingTokens/Source/NamingTokensUI/Private/SNamingTokensEditableTextBox.cpp
+++ b/Engine/Plugins/Developer/NamingTokens/Source/NamingTokensUI/Private/SNamingTokensEditableTextBox.cpp
@@ -207,7 +207,10 @@ void SNamingTokensEditableTextBox::Construct(const FArguments& InArgs)
 		SetVisibility(EVisibility::HitTestInvisible);
 	}
 
-	OnEditableTextCommitted(TokenizedText.Get(), ETextCommit::Default);
+	// Evaluate the initial text, validate it, and evaluate tokens
+	const FText InitialText = TokenizedText.Get();
+	ValidateTokenizedText(InitialText);
+	EvaluateNamingTokens();
 }
 
 void SNamingTokensEditableTextBox::SetTokenizedText(const FText& InText)
