NamingTokens: Fix token textbox autocomplete of closing braces when text is selected.
#rnx

diff --git a/Engine/Plugins/Developer/NamingTokens/Source/NamingTokensUI/Private/SNamingTokensEditableTextBox.cpp b/Engine/Plugins/Developer/NamingTokens/Source/NamingTokensUI/Private/SNamingTokensEditableTextBox.cpp
--- a/Engine/Plugins/Developer/NamingTokens/Source/NamingTokensUI/Private/SNamingTokensEditableTextBox.cpp
+++ b/Engine/Plugins/Developer/NamingTokens/Source/NamingTokensUI/Private/SNamingTokensEditableTextBox.cpp
@@ -766,6 +766,7 @@ FReply SNamingTokensEditableTextBox::OnEditableTextKeyChar(const FGeometry& MyGe
 	if (InCharacterEvent.GetCharacter() == TEXT('{'))
 	{
 		// Create a closing bracket.
+		EditableText->DeleteSelectedText();
 		const FTextLocation FilterCursorLocation = EditableText->GetCursorLocation();
 		EditableText->InsertTextAtCursor(TEXT("}"));
 		EditableText->GoTo(FilterCursorLocation);
