NamingTokens: Add editor utility support. Naming Tokens can now be created either as runtime or as editor utility blueprints.

Fix autocomplete not identifying BP children namespaces of custom native children.

Refactored editor only portion out of the UncookedOnly module.

#Synced-CL 48223680

diff --git a/Engine/Plugins/Developer/NamingTokens/NamingTokens.uplugin b/Engine/Plugins/Developer/NamingTokens/NamingTokens.uplugin
--- a/Engine/Plugins/Developer/NamingTokens/NamingTokens.uplugin
+++ b/Engine/Plugins/Developer/NamingTokens/NamingTokens.uplugin
@@ -19,14 +19,19 @@
 			"Type": "Runtime",
 			"LoadingPhase": "Default"
 		},
+		{
+			"Name": "NamingTokensUI",
+			"Type": "Runtime",
+			"LoadingPhase": "Default"
+		},
 		{
 			"Name": "NamingTokensUncookedOnly",
 			"Type": "UncookedOnly",
 			"LoadingPhase": "Default"
 		},
 		{
-			"Name": "NamingTokensUI",
-			"Type": "Runtime",
+			"Name": "NamingTokensEditor",
+			"Type": "Editor",
 			"LoadingPhase": "Default"
 		}
 	],
diff --git a/Engine/Plugins/Developer/NamingTokens/Source/NamingTokens/Private/NamingTokensEngineSubsystem.cpp b/Engine/Plugins/Developer/NamingTokens/Source/NamingTokens/Private/NamingTokensEngineSubsystem.cpp
--- a/Engine/Plugins/Developer/NamingTokens/Source/NamingTokens/Private/NamingTokensEngineSubsystem.cpp
+++ b/Engine/Plugins/Developer/NamingTokens/Source/NamingTokens/Private/NamingTokensEngineSubsystem.cpp
@@ -6,6 +6,7 @@
 #include "NamingTokensEvaluationData.h"
 #include "NamingTokensLog.h"
 #include "Utils/NamingTokenUtils.h"
+#include "Utils/NamingTokenUtilsPrivate.h"
 
 #include "AssetRegistry/AssetData.h"
 #include "AssetRegistry/AssetDataTagMap.h"
@@ -18,30 +19,22 @@
 
 TArray<FString> UNamingTokensEngineSubsystem::GetAllNamespaces(bool bIncludePrivateNamespaces) const
 {
-	const FAssetRegistryModule& AssetRegistryModule = FModuleManager::LoadModuleChecked<FAssetRegistryModule>("AssetRegistry");
-	const IAssetRegistry& AssetRegistry = AssetRegistryModule.Get();
-
-	// Find all assets whose native parent class is UNamingTokens, and add their namespaces
-	TArray<FAssetData> NamingTokenAssets;
-	const TMultiMap<FName, FString> TagValues = { { FBlueprintTags::NativeParentClassPath, FObjectPropertyBase::GetExportPath(UNamingTokens::StaticClass()) } };
-	AssetRegistry.GetAssetsByTagValues(TagValues, NamingTokenAssets);
-
 	TSet<FString> DiscoveredNamespaces;
-	for (const FAssetData& NamingTokenAssetData : NamingTokenAssets)
+	
+	UE::NamingTokens::Utils::Private::ForEachNamingTokensBlueprint(
+		[bIncludePrivateNamespaces, &DiscoveredNamespaces](const FAssetData& AssetData, const FTopLevelAssetPath& ClassObjectPath)
 	{
 		bool bIsPrivateNamespace = false;
-		NamingTokenAssetData.GetTagValue(UNamingTokens::GetIsPrivateNamespacePropertyName(), bIsPrivateNamespace);
+		AssetData.GetTagValue(UNamingTokens::GetIsPrivateNamespacePropertyName(), bIsPrivateNamespace);
 
-		if (bIsPrivateNamespace && !bIncludePrivateNamespaces)
+		if (!bIsPrivateNamespace || bIncludePrivateNamespaces)
 		{
-			continue;
-		}
-		
-		FString FoundNamespace;
-		NamingTokenAssetData.GetTagValue(UNamingTokens::GetNamespacePropertyName(), FoundNamespace);
+			FString FoundNamespace;
+			AssetData.GetTagValue(UNamingTokens::GetNamespacePropertyName(), FoundNamespace);
 
-		DiscoveredNamespaces.Add(FoundNamespace);
-	}
+			DiscoveredNamespaces.Add(FoundNamespace);
+		}
+	});
 
 	// Find all native classes derived from UNamingTokens, and add their namespaces
 	TArray<UClass*> DerivedClasses;
@@ -82,77 +75,52 @@ UNamingTokens* UNamingTokensEngineSubsystem::GetNamingTokens(const FString& InNa
 
 	FoundTokens = GetNamingTokensNative(InNamespace);
 
-	const UClass* TargetClass = UNamingTokens::StaticClass();
-
+	const FAssetRegistryModule& AssetRegistryModule = FModuleManager::LoadModuleChecked<FAssetRegistryModule>("AssetRegistry");
+	const IAssetRegistry& AssetRegistry = AssetRegistryModule.Get();
+	
 	// Check blueprint classes.
 	// NOTE: We may want to use OnAssetAdded and cache the asset or load the class then.
+	UE::NamingTokens::Utils::Private::ForEachNamingTokensBlueprint(
+		[&InNamespace, &AssetRegistry, &FoundTokens, this](const FAssetData& AssetData, const FTopLevelAssetPath& ClassObjectPath)
 	{
-		const FAssetRegistryModule& AssetRegistryModule = FModuleManager::LoadModuleChecked<FAssetRegistryModule>("AssetRegistry");
-		const IAssetRegistry& AssetRegistry = AssetRegistryModule.Get();
-
-		FARFilter Filter;
-		Filter.ClassPaths.Add(UBlueprint::StaticClass()->GetClassPathName());
-		Filter.bRecursiveClasses = true;
-		Filter.bRecursivePaths = true;
-
-		TArray<FAssetData> BlueprintAssets;
-		AssetRegistry.GetAssets(Filter, BlueprintAssets);
-
-		TSet<FTopLevelAssetPath> DerivedClassNames;
-		AssetRegistry.GetDerivedClassNames({ TargetClass->GetClassPathName() }, {}, DerivedClassNames);
-
-		for (const FAssetData& AssetData : BlueprintAssets)
+		// Now check for the namespace match.
+		FString FoundNamespace;
+		if (AssetData.GetTagValue(UNamingTokens::GetNamespacePropertyName(), FoundNamespace)
+			&& FoundNamespace == InNamespace)
 		{
-			// Narrow down to only our assets.
-			const FAssetDataTagMapSharedView::FFindTagResult Result = AssetData.TagsAndValues.FindTag(TEXT("GeneratedClass"));
-			if (Result.IsSet())
+			// Allow BP children to overwrite native classes.
+			auto IsAssetChildOfNativeToken = [&AssetRegistry, &InNamespace](const FTopLevelAssetPath& InAssetPath, const UNamingTokens* InParentTokens) -> bool
 			{
-				const FString& GeneratedClassPathPtr = Result.GetValue();
-				const FTopLevelAssetPath ClassObjectPath(FPackageName::ExportTextPathToObjectPath(*GeneratedClassPathPtr));
-
-				if (DerivedClassNames.Contains(ClassObjectPath))
+				check(InParentTokens);
+				if (InParentTokens->GetClass()->IsNative())
 				{
-					// Now check for the namespace match.
-					FString FoundNamespace;
-					if (AssetData.GetTagValue(UNamingTokens::GetNamespacePropertyName(), FoundNamespace)
-						&& FoundNamespace == InNamespace)
+					// We need to make sure this BP is a child of the chosen native class.
+					TSet<FTopLevelAssetPath> NativeDerivedClassNames;
+					AssetRegistry.GetDerivedClassNames({ InParentTokens->GetClass()->GetClassPathName() },
+						{}, NativeDerivedClassNames);
+
+					if (NativeDerivedClassNames.Contains(InAssetPath))
 					{
-						// Allow BP children to overwrite native classes.
-						auto IsAssetChildOfNativeToken = [&AssetRegistry, &InNamespace](const FTopLevelAssetPath& InAssetPath, const UNamingTokens* InParentTokens) -> bool
-						{
-							check(InParentTokens);
-							if (InParentTokens->GetClass()->IsNative())
-							{
-								// We need to make sure this BP is a child of the chosen native class.
-								TSet<FTopLevelAssetPath> NativeDerivedClassNames;
-								AssetRegistry.GetDerivedClassNames({ InParentTokens->GetClass()->GetClassPathName() },
-									{}, NativeDerivedClassNames);
-
-								if (NativeDerivedClassNames.Contains(InAssetPath))
-								{
-									UE_LOG(LogNamingTokens, Log, TEXT("Using namespace '%s' of BP child '%s' instead of native parent '%s'."),
-										*InNamespace, *InAssetPath.ToString(), *InParentTokens->GetClass()->GetName());
-									return true;
-								}
-							}
-							return false;
-						};
-						
-						if (!FoundTokens || IsAssetChildOfNativeToken(ClassObjectPath, FoundTokens))
-						{
-							const TSoftClassPtr<UNamingTokens> SoftClassPath(ClassObjectPath.ToString());
-							FoundTokens = LoadNamingToken(SoftClassPath, InNamespace);
-						}
-						else
-						{
-							UE_LOG(LogNamingTokens, Warning, TEXT("Namespace '%s' exists more than once in class '%s' and BP asset '%s'."),
-								*InNamespace, *FoundTokens->GetClass()->GetName(), *AssetData.PackageName.ToString());
-						}
+						UE_LOG(LogNamingTokens, Log, TEXT("Using namespace '%s' of BP child '%s' instead of native parent '%s'."),
+							*InNamespace, *InAssetPath.ToString(), *InParentTokens->GetClass()->GetName());
+						return true;
 					}
 				}
+				return false;
+			};
+			
+			if (!FoundTokens || IsAssetChildOfNativeToken(ClassObjectPath, FoundTokens))
+			{
+				const TSoftClassPtr<UNamingTokens> SoftClassPath(ClassObjectPath.ToString());
+				FoundTokens = LoadNamingToken(SoftClassPath, InNamespace);
+			}
+			else
+			{
+				UE_LOG(LogNamingTokens, Warning, TEXT("Namespace '%s' exists more than once in class '%s' and BP asset '%s'."),
+					*InNamespace, *FoundTokens->GetClass()->GetName(), *AssetData.PackageName.ToString());
 			}
 		}
-	}
+	});
 
 	return FoundTokens;
 }
diff --git a/Engine/Plugins/Developer/NamingTokens/Source/NamingTokens/Private/Utils/NamingTokenUtilsPrivate.cpp b/Engine/Plugins/Developer/NamingTokens/Source/NamingTokens/Private/Utils/NamingTokenUtilsPrivate.cpp
new file mode 100644
--- /dev/null
+++ b/Engine/Plugins/Developer/NamingTokens/Source/NamingTokens/Private/Utils/NamingTokenUtilsPrivate.cpp
@@ -0,0 +1,44 @@
+// Copyright Epic Games, Inc. All Rights Reserved.
+
+#include "NamingTokenUtilsPrivate.h"
+
+#include "AssetRegistry/AssetData.h"
+#include "AssetRegistry/AssetDataTagMap.h"
+#include "AssetRegistry/AssetRegistryModule.h"
+#include "Engine/Blueprint.h"
+#include "Misc/PackageName.h"
+#include "NamingTokens.h"
+
+void UE::NamingTokens::Utils::Private::ForEachNamingTokensBlueprint(
+	const TFunction<void(const FAssetData& AssetData, const FTopLevelAssetPath& ClassObjectPath)>& Function)
+{
+	const FAssetRegistryModule& AssetRegistryModule = FModuleManager::LoadModuleChecked<FAssetRegistryModule>("AssetRegistry");
+	const IAssetRegistry& AssetRegistry = AssetRegistryModule.Get();
+	
+	FARFilter Filter;
+	Filter.ClassPaths.Add(UBlueprint::StaticClass()->GetClassPathName());
+	Filter.bRecursiveClasses = true;
+	Filter.bRecursivePaths = true;
+
+	TArray<FAssetData> BlueprintAssets;
+	AssetRegistry.GetAssets(Filter, BlueprintAssets);
+
+	TSet<FTopLevelAssetPath> DerivedClassNames;
+	AssetRegistry.GetDerivedClassNames({ UNamingTokens::StaticClass()->GetClassPathName() }, {}, DerivedClassNames);
+	
+	for (const FAssetData& NamingTokenAssetData : BlueprintAssets)
+	{
+		// Narrow down to only our assets.
+		const FAssetDataTagMapSharedView::FFindTagResult Result = NamingTokenAssetData.TagsAndValues.FindTag(TEXT("GeneratedClass"));
+		if (Result.IsSet())
+		{
+			const FString& GeneratedClassPathPtr = Result.GetValue();
+			const FTopLevelAssetPath ClassObjectPath(FPackageName::ExportTextPathToObjectPath(*GeneratedClassPathPtr));
+
+			if (DerivedClassNames.Contains(ClassObjectPath))
+			{
+				Function(NamingTokenAssetData, ClassObjectPath);
+			}
+		}
+	}
+}
diff --git a/Engine/Plugins/Developer/NamingTokens/Source/NamingTokens/Private/Utils/NamingTokenUtilsPrivate.h b/Engine/Plugins/Developer/NamingTokens/Source/NamingTokens/Private/Utils/NamingTokenUtilsPrivate.h
new file mode 100644
--- /dev/null
+++ b/Engine/Plugins/Developer/NamingTokens/Source/NamingTokens/Private/Utils/NamingTokenUtilsPrivate.h
@@ -0,0 +1,14 @@
+// Copyright Epic Games, Inc. All Rights Reserved.
+
+#pragma once
+
+#include "Templates/Function.h"
+
+struct FAssetData;
+struct FTopLevelAssetPath;
+
+namespace UE::NamingTokens::Utils::Private
+{
+	/** Iterate each Naming Tokens blueprint asset, executing a function for each qualifying asset. */
+	void ForEachNamingTokensBlueprint(const TFunction<void(const FAssetData& AssetData, const FTopLevelAssetPath& ClassObjectPath)>& Function);
+}
diff --git a/Engine/Plugins/Developer/NamingTokens/Source/NamingTokens/Public/NamingTokens.h b/Engine/Plugins/Developer/NamingTokens/Source/NamingTokens/Public/NamingTokens.h
--- a/Engine/Plugins/Developer/NamingTokens/Source/NamingTokens/Public/NamingTokens.h
+++ b/Engine/Plugins/Developer/NamingTokens/Source/NamingTokens/Public/NamingTokens.h
@@ -122,6 +122,10 @@ protected:
 
 public:
 	static FName GetProcessTokenTemplateFunctionName() { return GET_FUNCTION_NAME_CHECKED(UNamingTokens, ProcessTokenTemplateFunction); } 
+	
+	static FName GetOnPreEvaluateFunctionName() { return GET_FUNCTION_NAME_CHECKED(UNamingTokens, OnPreEvaluate); }
+	static FName GetOnPostEvaluateFunctionName() { return GET_FUNCTION_NAME_CHECKED(UNamingTokens, OnPostEvaluate); }
+	
 private:
 	/** Template function for us to dynamically create subclass graphs from matching this signature. */
 	UFUNCTION()
@@ -216,8 +220,6 @@ private:
 	UPROPERTY(VisibleDefaultsOnly, Transient, Category = "NamingTokens")
 	FText TestTokenResult;
 #endif
-
-	friend class UNamingTokensFactory;
 };
 
 #undef UE_API
diff --git a/Engine/Plugins/Developer/NamingTokens/Source/NamingTokensEditor/NamingTokensEditor.Build.cs b/Engine/Plugins/Developer/NamingTokens/Source/NamingTokensEditor/NamingTokensEditor.Build.cs
new file mode 100644
--- /dev/null
+++ b/Engine/Plugins/Developer/NamingTokens/Source/NamingTokensEditor/NamingTokensEditor.Build.cs
@@ -0,0 +1,36 @@
+﻿// Copyright Epic Games, Inc. All Rights Reserved.
+
+using UnrealBuildTool;
+
+public class NamingTokensEditor : ModuleRules
+{
+    public NamingTokensEditor(ReadOnlyTargetRules Target) : base(Target)
+    {
+        PCHUsage = ModuleRules.PCHUsageMode.UseExplicitOrSharedPCHs;
+
+        PublicDependencyModuleNames.AddRange(
+            new string[]
+            {
+                "Core",
+                "NamingTokens",
+            }
+        );
+
+        PrivateDependencyModuleNames.AddRange(
+            new string[]
+            {
+	            "AssetDefinition",
+	            "BlueprintGraph",
+                "CoreUObject",
+                "EditorToolEvents",
+                "Engine",
+                "EngineAssetDefinitions",
+                "InputCore",
+                "Slate",
+                "SlateCore",
+                "Projects",
+                "UnrealEd",
+            }
+        );
+    }
+}
\ No newline at end of file
diff --git a/Engine/Plugins/Developer/NamingTokens/Source/NamingTokensEditor/Private/AssetDefinition_NamingTokens.cpp b/Engine/Plugins/Developer/NamingTokens/Source/NamingTokensEditor/Private/AssetDefinition_NamingTokens.cpp
new file mode 100644
--- /dev/null
+++ b/Engine/Plugins/Developer/NamingTokens/Source/NamingTokensEditor/Private/AssetDefinition_NamingTokens.cpp
@@ -0,0 +1,29 @@
+// Copyright Epic Games, Inc. All Rights Reserved.
+
+#include "AssetDefinition_NamingTokens.h"
+
+#include "NamingTokens.h"
+#include "NamingTokensStyle.h"
+
+#include UE_INLINE_GENERATED_CPP_BY_NAME(AssetDefinition_NamingTokens)
+
+TSoftClassPtr<UObject> UAssetDefinition_NamingTokens::GetAssetClass() const
+{
+	return UNamingTokens::StaticClass();
+}
+
+TConstArrayView<FAssetCategoryPath> UAssetDefinition_NamingTokens::GetAssetCategories() const
+{
+	static const FAssetCategoryPath Categories[] = { FAssetCategoryPath(EAssetCategoryPaths::Misc) };
+	return Categories;
+}
+
+FText UAssetDefinition_NamingTokens::GetAssetDisplayName() const
+{
+	return NSLOCTEXT("AssetDefinition_NamingTokens", "AssetDefinition_NamingTokens_DisplayName", "Naming Tokens");
+}
+
+const FSlateBrush* UAssetDefinition_NamingTokens::GetThumbnailBrush(const FAssetData& InAssetData, const FName InClassName) const
+{
+	return FNamingTokensStyle::Get().GetBrush("ClassThumbnail.NamingTokens");
+}
diff --git a/Engine/Plugins/Developer/NamingTokens/Source/NamingTokensEditor/Private/AssetDefinition_NamingTokens.h b/Engine/Plugins/Developer/NamingTokens/Source/NamingTokensEditor/Private/AssetDefinition_NamingTokens.h
new file mode 100644
--- /dev/null
+++ b/Engine/Plugins/Developer/NamingTokens/Source/NamingTokensEditor/Private/AssetDefinition_NamingTokens.h
@@ -0,0 +1,19 @@
+// Copyright Epic Games, Inc. All Rights Reserved.
+
+#pragma once
+
+#include "Script/AssetDefinition_Blueprint.h"
+
+#include "AssetDefinition_NamingTokens.generated.h"
+
+UCLASS()
+class UAssetDefinition_NamingTokens : public UAssetDefinition_Blueprint
+{
+	GENERATED_BODY()
+	
+public:
+	virtual TSoftClassPtr<UObject> GetAssetClass() const override;
+	virtual TConstArrayView<FAssetCategoryPath> GetAssetCategories() const override;
+	virtual FText GetAssetDisplayName() const override;
+	virtual const FSlateBrush* GetThumbnailBrush(const FAssetData& InAssetData, const FName InClassName) const override;
+};
diff --git a/Engine/Plugins/Developer/NamingTokens/Source/NamingTokensEditor/Private/Customization/NamingTokensCustomization.cpp b/Engine/Plugins/Developer/NamingTokens/Source/NamingTokensEditor/Private/Customization/NamingTokensCustomization.cpp
new file mode 100644
--- /dev/null
+++ b/Engine/Plugins/Developer/NamingTokens/Source/NamingTokensEditor/Private/Customization/NamingTokensCustomization.cpp
@@ -0,0 +1,337 @@
+// Copyright Epic Games, Inc. All Rights Reserved.
+
+#include "NamingTokensCustomization.h"
+
+#include "NamingTokenData.h"
+#include "NamingTokens.h"
+#include "Utils/NamingTokensEditorUtils.h"
+#include "Utils/NamingTokenUtils.h"
+
+#include "DetailLayoutBuilder.h"
+#include "DetailWidgetRow.h"
+#include "Engine/Blueprint.h"
+#include "IDetailChildrenBuilder.h"
+#include "IDetailPropertyRow.h"
+#include "IPropertyUtilities.h"
+#include "ScopedTransaction.h"
+#include "Widgets/Input/SButton.h"
+#include "Widgets/Input/SComboBox.h"
+#include "Widgets/Text/STextBlock.h"
+
+#define LOCTEXT_NAMESPACE "NamingTokensCustomization"
+
+namespace UE::NamingTokens::Customization::Private
+{
+	// If the error message should be visible.
+	EVisibility GetErrorVisibilityFromProperty(const TSharedPtr<IPropertyHandle>& InPropertyHandle, const TSharedPtr<FText>& InErrorMessage)
+	{
+		if (InPropertyHandle.IsValid() && InErrorMessage.IsValid())
+		{
+			FString NamespaceValue;
+			if (InPropertyHandle->GetValue(NamespaceValue) == FPropertyAccess::Result::Success
+				&& Utils::ValidateName(NamespaceValue, *InErrorMessage))
+			{
+				return EVisibility::Collapsed;
+			}
+		}
+
+		return EVisibility::Visible;
+	}
+
+	// Construct the error tooltip message.
+	FText CreateErrorTooltipMessage(const TSharedPtr<FText>& InErrorMessage)
+	{
+		if (InErrorMessage.IsValid() && !InErrorMessage->IsEmpty())
+		{
+			FFormatNamedArguments Args;
+			Args.Add(TEXT("Error"), *InErrorMessage);
+			return FText::Format(LOCTEXT("ValueError", "Error: {Error}. Alphanumeric and '_' characters are allowed."), Args);
+		}
+		return FText::GetEmpty();
+	}
+	
+	// Creates the standard row widget for a property, adding in an error icon for validation.
+	void CreateRowWidgetWithError(IDetailPropertyRow& InRow, const TSharedPtr<IPropertyHandle>& InPropertyHandle,
+		const TSharedPtr<FText>& InErrorMessage)
+	{
+		InRow.CustomWidget()
+		.NameContent()
+		[
+			InPropertyHandle->CreatePropertyNameWidget()
+		]
+		.ValueContent()
+		.MaxDesiredWidth(250.0f)
+		[
+			SNew(SHorizontalBox)
+			+ SHorizontalBox::Slot()
+			.FillWidth(1.0f)
+			[
+				InPropertyHandle->CreatePropertyValueWidget()
+			]
+			+ SHorizontalBox::Slot()
+			.AutoWidth()
+			.VAlign(VAlign_Center)
+			.HAlign(HAlign_Center)
+			.Padding(2.f, 0.f, 0.f, 0.f)
+			[
+				SNew(SImage)
+				.Image(FAppStyle::Get().GetBrush("Icons.Error"))
+				.ColorAndOpacity(FAppStyle::Get().GetSlateColor("Colors.AccentRed"))
+				.ToolTipText_Lambda([InErrorMessage]()
+				{
+					return CreateErrorTooltipMessage(InErrorMessage);
+				})
+				.Visibility_Lambda([InPropertyHandle, InErrorMessage]()
+				{
+					return GetErrorVisibilityFromProperty(InPropertyHandle, InErrorMessage);
+				})
+			]
+		];
+	}
+}
+
+void FNamingTokensCustomization::CustomizeDetails(IDetailLayoutBuilder& DetailBuilder)
+{
+	NamespacePropertyHandle = DetailBuilder.GetProperty(UNamingTokens::GetNamespacePropertyName());
+	TokenKeyErrorMessage = MakeShared<FText>();
+	if (NamespacePropertyHandle.IsValid())
+	{
+		if (IDetailPropertyRow* Row = DetailBuilder.EditDefaultProperty(NamespacePropertyHandle))
+		{
+			UE::NamingTokens::Customization::Private::CreateRowWidgetWithError(*Row, NamespacePropertyHandle, TokenKeyErrorMessage);
+		}
+	}
+}
+
+void FNamingTokensDataCustomization::CustomizeHeader(TSharedRef<IPropertyHandle> PropertyHandle, FDetailWidgetRow& HeaderRow,
+                                                     IPropertyTypeCustomizationUtils& CustomizationUtils)
+{
+	CustomizationUtilsPtr = &CustomizationUtils;
+	NamespaceErrorMessage = MakeShared<FText>();
+	
+	if (const UNamingTokens* OwningTokens = GetOwningNamingTokens())
+	{
+		OwningBlueprint = UBlueprint::GetBlueprintFromClass(OwningTokens->GetClass());
+		if (OwningBlueprint.IsValid())
+		{
+			FunctionNameHandle = PropertyHandle->GetChildHandle(GET_MEMBER_NAME_CHECKED(FNamingTokenData, FunctionName));
+			check(FunctionNameHandle.IsValid());
+
+			TokenKeyHandle = PropertyHandle->GetChildHandle(GET_MEMBER_NAME_CHECKED(FNamingTokenData, TokenKey));
+			check(TokenKeyHandle.IsValid());
+		}
+	}
+
+	// If owning blueprint is null, then we may not be customizing the naming tokens directly, such as a settings object,
+	// or the tokens are predefined and can't be extended in BP.
+	
+	HeaderRow
+	.NameContent()
+	[
+		PropertyHandle->CreatePropertyNameWidget()
+	]
+	.ValueContent()
+	[
+		PropertyHandle->CreatePropertyValueWidget(/* bDisplayDefaultPropertyButtons (adds identical options) */ false)
+	];
+}
+
+void FNamingTokensDataCustomization::CustomizeChildren(TSharedRef<IPropertyHandle> PropertyHandle, IDetailChildrenBuilder& ChildBuilder,
+	IPropertyTypeCustomizationUtils& CustomizationUtils)
+{
+	if (!FunctionNameHandle.IsValid() || !TokenKeyHandle.IsValid())
+	{
+		return;
+	}
+	
+	const bool bShouldApplyCustomization = OwningBlueprint.IsValid();
+
+	if (bShouldApplyCustomization)
+	{
+		const TArray<UFunction*> Functions = GetAvailableFunctions();
+		for (const UFunction* Function : Functions)
+		{
+			FunctionNames.Add(MakeShared<FString>(Function->GetName()));
+		}
+
+		FString CurrentValue;
+		FunctionNameHandle->GetValue(CurrentValue);
+
+		// Default selected function
+		for (const TSharedPtr<FString>& Name : FunctionNames)
+		{
+			if (*Name == CurrentValue)
+			{
+				SelectedFunctionName = Name;
+				break;
+			}
+		}
+
+		FunctionNameHandle->MarkHiddenByCustomization();
+	}
+	
+	uint32 NumChildren = 0;
+	PropertyHandle->GetNumChildren(NumChildren);
+	
+	for (uint32 ChildNum = 0; ChildNum < NumChildren; ++ChildNum)
+	{
+		const TSharedPtr<IPropertyHandle> ChildHandle = PropertyHandle->GetChildHandle(ChildNum);
+		if (!ChildHandle.IsValid() || (bShouldApplyCustomization && ChildHandle->GetProperty() == FunctionNameHandle->GetProperty()))
+		{
+			continue;
+		}
+		IDetailPropertyRow& Row = ChildBuilder.AddProperty(ChildHandle.ToSharedRef());
+
+		if (ChildHandle->GetProperty() == TokenKeyHandle->GetProperty())
+		{
+			UE::NamingTokens::Customization::Private::CreateRowWidgetWithError(Row, TokenKeyHandle, NamespaceErrorMessage);
+		}
+	}
+
+	if (!bShouldApplyCustomization)
+	{
+		return;
+	}
+	
+	FDetailWidgetRow& HeaderRow = ChildBuilder.AddProperty(FunctionNameHandle.ToSharedRef()).CustomWidget();
+	HeaderRow.NameContent()
+	[
+		FunctionNameHandle->CreatePropertyNameWidget()
+	];
+	HeaderRow.ValueContent()
+	[
+		SNew(SHorizontalBox)
+		+SHorizontalBox::Slot()
+		[
+			SNew(SComboBox<TSharedPtr<FString>>)
+			.OptionsSource(&FunctionNames)
+			.OnGenerateWidget(this, &FNamingTokensDataCustomization::MakeComboBoxWidget)
+			.OnSelectionChanged(this, &FNamingTokensDataCustomization::OnFunctionSelected, FunctionNameHandle)
+			.InitiallySelectedItem(SelectedFunctionName)
+			[
+				SNew(STextBlock)
+				.Text(this, &FNamingTokensDataCustomization::GetSelectedFunctionText)
+			]
+		]
+		+ SHorizontalBox::Slot()
+		.AutoWidth()
+		.VAlign(VAlign_Center)
+		[
+			SNew(SButton)
+			[
+				SNew(SImage)
+				.Image(FAppStyle::Get().GetBrush("PListEditor.Button_AddToArray"))
+			]
+			.ButtonStyle(FAppStyle::Get(), "HoverHintOnly")
+			.ToolTipText(LOCTEXT("AddFunction", "Create and assign a new function graph for evaluating tokens.\nThis requires a valid Token Key entered."))
+			.OnClicked(this, &FNamingTokensDataCustomization::OnAddFunctionClicked)
+			.IsEnabled(this, &FNamingTokensDataCustomization::CanAddFunction)
+		]
+	];
+}
+
+UNamingTokens* FNamingTokensDataCustomization::GetOwningNamingTokens() const
+{
+	if (CustomizationUtilsPtr)
+	{
+		const TSharedPtr<IPropertyUtilities> PropertyUtilities = CustomizationUtilsPtr->GetPropertyUtilities();
+		if (PropertyUtilities.IsValid())
+		{
+			const TArray<TWeakObjectPtr<UObject>> ObjectsBeingCustomized = PropertyUtilities->GetSelectedObjects();
+
+			if (ObjectsBeingCustomized.Num() == 1)
+			{
+				if (UNamingTokens* Tokens = Cast<UNamingTokens>(ObjectsBeingCustomized[0]))
+				{
+					return Tokens;
+				}
+			}
+		}
+	}
+	return nullptr;
+}
+
+TArray<UFunction*> FNamingTokensDataCustomization::GetAvailableFunctions() const
+{
+	TArray<UFunction*> Result;
+	if (OwningBlueprint.IsValid())
+	{
+		for (TFieldIterator<UFunction> It(OwningBlueprint->SkeletonGeneratedClass ?
+			OwningBlueprint->SkeletonGeneratedClass : OwningBlueprint->GeneratedClass,
+			EFieldIteratorFlags::IncludeSuper); It; ++It)
+		{
+			if (UE::NamingTokens::Utils::ValidateTokenFunction(*It))
+			{
+				Result.Add(*It);
+			}
+		}
+	}
+
+	return Result;
+}
+
+TSharedRef<SWidget> FNamingTokensDataCustomization::MakeComboBoxWidget(TSharedPtr<FString> InItem)
+{
+	return SNew(STextBlock)
+		.Text(FText::FromString(*InItem));
+}
+
+void FNamingTokensDataCustomization::OnFunctionSelected(TSharedPtr<FString> NewValue, ESelectInfo::Type SelectInfo,
+	TSharedPtr<IPropertyHandle> InFunctionNameHandle)
+{
+	if (NewValue.IsValid())
+	{
+		SelectedFunctionName = NewValue;
+		InFunctionNameHandle->SetValue(*NewValue);
+	}
+}
+
+FText FNamingTokensDataCustomization::GetSelectedFunctionText() const
+{
+	return SelectedFunctionName.IsValid() ? FText::FromString(*SelectedFunctionName) : FText::FromString(TEXT("Select Function"));
+}
+
+FReply FNamingTokensDataCustomization::OnAddFunctionClicked()
+{
+	if (UBlueprint* Blueprint = OwningBlueprint.Get())
+	{
+		FString TokenKey;
+		TokenKeyHandle->GetValue(TokenKey);
+		FScopedTransaction Transaction(LOCTEXT("CreateTokenGraph", "Create Token Graph"));
+
+		const FName NewFunctionName = UE::NamingTokens::Utils::Editor::Private::CreateNewTokenGraph(Blueprint, TokenKey);
+		FunctionNameHandle->SetValue(NewFunctionName);
+	}
+	
+	return FReply::Handled();
+}
+
+bool FNamingTokensDataCustomization::CanAddFunction() const
+{
+	FString TokenKey;
+	if (TokenKeyHandle.IsValid())
+	{
+		TokenKeyHandle->GetValue(TokenKey);
+
+		FText ErrorMessage;
+		if (!UE::NamingTokens::Utils::ValidateName(TokenKey, ErrorMessage))
+		{
+			return false;
+		}
+			
+		const FString BaseFunctionName = UE::NamingTokens::Utils::Editor::Private::CreateBaseTokenFunctionName(TokenKey);
+
+		// Make sure the name isn't already being used.
+		for (const TSharedPtr<FString>& FunctionName : FunctionNames)
+		{
+			if (FunctionName.IsValid() && **FunctionName == BaseFunctionName)
+			{
+				return false;
+			}
+		}
+	}
+	
+	return !TokenKey.IsEmpty();
+}
+
+#undef LOCTEXT_NAMESPACE
diff --git a/Engine/Plugins/Developer/NamingTokens/Source/NamingTokensEditor/Private/Customization/NamingTokensCustomization.h b/Engine/Plugins/Developer/NamingTokens/Source/NamingTokensEditor/Private/Customization/NamingTokensCustomization.h
new file mode 100644
--- /dev/null
+++ b/Engine/Plugins/Developer/NamingTokens/Source/NamingTokensEditor/Private/Customization/NamingTokensCustomization.h
@@ -0,0 +1,78 @@
+// Copyright Epic Games, Inc. All Rights Reserved.
+
+#pragma once
+
+#include "CoreMinimal.h"
+#include "IDetailCustomization.h"
+#include "Input/Reply.h"
+#include "IPropertyTypeCustomization.h"
+#include "Layout/Visibility.h"
+
+class UNamingTokens;
+class SWidget;
+
+class FNamingTokensCustomization final : public IDetailCustomization
+{
+public:
+	static TSharedRef<IDetailCustomization> MakeInstance()
+	{
+		return MakeShareable(new FNamingTokensCustomization);
+	}
+
+	virtual void CustomizeDetails(IDetailLayoutBuilder& DetailBuilder) override;
+	
+private:
+	/** Handle to namespace property. */
+	TSharedPtr<IPropertyHandle> NamespacePropertyHandle;
+	/** Stores error text. */
+	TSharedPtr<FText> TokenKeyErrorMessage;
+};
+
+class FNamingTokensDataCustomization final : public IPropertyTypeCustomization
+{
+public:
+	static TSharedRef<IPropertyTypeCustomization> MakeInstance()
+	{
+		return MakeShareable(new FNamingTokensDataCustomization());
+	}
+
+	virtual void CustomizeHeader(TSharedRef<IPropertyHandle> PropertyHandle, FDetailWidgetRow& HeaderRow, IPropertyTypeCustomizationUtils& CustomizationUtils) override;
+	virtual void CustomizeChildren(TSharedRef<IPropertyHandle> PropertyHandle, IDetailChildrenBuilder& ChildBuilder, IPropertyTypeCustomizationUtils& CustomizationUtils) override;
+
+	/** The naming tokens owning this customization. */
+	UNamingTokens* GetOwningNamingTokens() const;
+
+private:
+	/** Available functions which can be assigned to naming tokens. */
+	TArray<UFunction*> GetAvailableFunctions() const;
+	/** Make the combo box for selecting functions. */
+	TSharedRef<SWidget> MakeComboBoxWidget(TSharedPtr<FString> InItem);
+	/** When the user selects a function. */
+	void OnFunctionSelected(
+		TSharedPtr<FString> NewValue, 
+		ESelectInfo::Type SelectInfo, 
+		TSharedPtr<IPropertyHandle> InFunctionNameHandle);
+	/** The text of the currently selected function. */
+	FText GetSelectedFunctionText() const;
+	/** User clicked the button to add a new function. */
+	FReply OnAddFunctionClicked();
+	/** If a new function can be added. */
+	bool CanAddFunction() const;
+
+private:
+	/** Customization pointer. */
+	IPropertyTypeCustomizationUtils* CustomizationUtilsPtr = nullptr;
+	/** Blueprint of the naming tokens. */
+	TWeakObjectPtr<UBlueprint> OwningBlueprint;
+
+	/** All available function names. */
+	TArray<TSharedPtr<FString>> FunctionNames;
+	/** The selected function name. */
+	TSharedPtr<FString> SelectedFunctionName;
+	/** Property handle to the function name. */
+	TSharedPtr<IPropertyHandle> FunctionNameHandle;
+	/** Property handle to the token key. */
+	TSharedPtr<IPropertyHandle> TokenKeyHandle;
+	/** Stores error text. */
+	TSharedPtr<FText> NamespaceErrorMessage;
+};
\ No newline at end of file
diff --git a/Engine/Plugins/Developer/NamingTokens/Source/NamingTokensEditor/Private/EditorUtilityNamingTokens.cpp b/Engine/Plugins/Developer/NamingTokens/Source/NamingTokensEditor/Private/EditorUtilityNamingTokens.cpp
new file mode 100644
--- /dev/null
+++ b/Engine/Plugins/Developer/NamingTokens/Source/NamingTokensEditor/Private/EditorUtilityNamingTokens.cpp
@@ -0,0 +1,23 @@
+// Copyright Epic Games, Inc. All Rights Reserved.
+
+#include "EditorUtilityNamingTokens.h"
+
+#include "EditorToolDelegates.h"
+
+#include UE_INLINE_GENERATED_CPP_BY_NAME(EditorUtilityNamingTokens)
+
+UEditorUtilityNamingTokens::UEditorUtilityNamingTokens()
+{
+}
+
+void UEditorUtilityNamingTokens::PostLoad()
+{
+	Super::PostLoad();
+	// UEditorUtilityObject::PostLoad() calls this, but we don't extend UEditorUtilityObject, so we call this for consistency.
+	FEditorToolDelegates::OnEditorToolStarted.Broadcast(GetClass()->GetName());
+}
+
+bool UEditorUtilityNamingTokens::IsEditorOnly() const
+{
+	return true;
+}
diff --git a/Engine/Plugins/Developer/NamingTokens/Source/NamingTokensEditor/Private/NamingTokensEditorModule.cpp b/Engine/Plugins/Developer/NamingTokens/Source/NamingTokensEditor/Private/NamingTokensEditorModule.cpp
new file mode 100644
--- /dev/null
+++ b/Engine/Plugins/Developer/NamingTokens/Source/NamingTokensEditor/Private/NamingTokensEditorModule.cpp
@@ -0,0 +1,41 @@
+﻿// Copyright Epic Games, Inc. All Rights Reserved.
+
+#include "NamingTokensEditorModule.h"
+
+#include "Customization/NamingTokensCustomization.h"
+#include "NamingTokens.h"
+#include "NamingTokensStyle.h"
+#include "PropertyEditorModule.h"
+
+#define LOCTEXT_NAMESPACE "FNamingTokensEditorModule"
+
+void FNamingTokensEditorModule::StartupModule()
+{
+	FPropertyEditorModule& PropertyModule = FModuleManager::GetModuleChecked<FPropertyEditorModule>("PropertyEditor");
+
+	PropertyModule.RegisterCustomClassLayout(
+		UNamingTokens::StaticClass()->GetFName(),
+		FOnGetDetailCustomizationInstance::CreateStatic(&FNamingTokensCustomization::MakeInstance)
+	);
+
+	PropertyModule.RegisterCustomPropertyTypeLayout(FNamingTokenData::StaticStruct()->GetFName(),
+		FOnGetPropertyTypeCustomizationInstance::CreateStatic(&FNamingTokensDataCustomization::MakeInstance));
+
+	FNamingTokensStyle::Get();
+}
+
+void FNamingTokensEditorModule::ShutdownModule()
+{
+	if (UObjectInitialized() && !IsEngineExitRequested())
+	{
+		if (FPropertyEditorModule* PropertyEditorModule = FModuleManager::GetModulePtr<FPropertyEditorModule>("PropertyEditor"))
+		{
+			PropertyEditorModule->UnregisterCustomClassLayout(UNamingTokens::StaticClass()->GetFName());
+			PropertyEditorModule->UnregisterCustomPropertyTypeLayout(FNamingTokenData::StaticStruct()->GetFName());
+		}
+	}
+}
+
+#undef LOCTEXT_NAMESPACE
+    
+IMPLEMENT_MODULE(FNamingTokensEditorModule, NamingTokensEditor)
\ No newline at end of file
diff --git a/Engine/Plugins/Developer/NamingTokens/Source/NamingTokensEditor/Private/NamingTokensEditorModule.h b/Engine/Plugins/Developer/NamingTokens/Source/NamingTokensEditor/Private/NamingTokensEditorModule.h
new file mode 100644
--- /dev/null
+++ b/Engine/Plugins/Developer/NamingTokens/Source/NamingTokensEditor/Private/NamingTokensEditorModule.h
@@ -0,0 +1,13 @@
+﻿// Copyright Epic Games, Inc. All Rights Reserved.
+
+#pragma once
+
+#include "CoreMinimal.h"
+#include "Modules/ModuleManager.h"
+
+class FNamingTokensEditorModule : public IModuleInterface
+{
+public:
+    virtual void StartupModule() override;
+    virtual void ShutdownModule() override;
+};
diff --git a/Engine/Plugins/Developer/NamingTokens/Source/NamingTokensEditor/Private/NamingTokensStyle.cpp b/Engine/Plugins/Developer/NamingTokens/Source/NamingTokensEditor/Private/NamingTokensStyle.cpp
new file mode 100644
--- /dev/null
+++ b/Engine/Plugins/Developer/NamingTokens/Source/NamingTokensEditor/Private/NamingTokensStyle.cpp
@@ -0,0 +1,37 @@
+// Copyright Epic Games, Inc. All Rights Reserved.
+
+#include "NamingTokensStyle.h"
+
+#include "Interfaces/IPluginManager.h"
+#include "Styling/CoreStyle.h"
+#include "Styling/SlateStyleMacros.h"
+#include "Styling/SlateTypes.h"
+
+FName FNamingTokensStyle::StyleName("NamingTokensStyle");
+
+FNamingTokensStyle& FNamingTokensStyle::Get()
+{
+	static FNamingTokensStyle Inst;
+	return Inst;
+}
+
+FNamingTokensStyle::FNamingTokensStyle()
+	: FSlateStyleSet(StyleName)
+{
+	// Icons
+	const FVector2D Icon16x16(16.0f, 16.0f);
+	const FVector2D Icon64x64(64.0f, 64.0f);
+
+	const FString ContentDir = IPluginManager::Get().FindPlugin(TEXT("NamingTokens"))->GetContentDir();
+	SetContentRoot(ContentDir);
+
+	Set("ClassIcon.NamingTokens", new FSlateVectorImageBrush(RootToContentDir(TEXT("Slate/Token_16.svg")), Icon16x16));
+	Set("ClassThumbnail.NamingTokens", new FSlateVectorImageBrush(RootToContentDir(TEXT("Slate/Token_64.svg")), Icon64x64));
+
+	FSlateStyleRegistry::RegisterSlateStyle(*this);
+}
+
+FNamingTokensStyle::~FNamingTokensStyle()
+{
+	FSlateStyleRegistry::UnRegisterSlateStyle(*this);
+}
diff --git a/Engine/Plugins/Developer/NamingTokens/Source/NamingTokensEditor/Private/NamingTokensStyle.h b/Engine/Plugins/Developer/NamingTokens/Source/NamingTokensEditor/Private/NamingTokensStyle.h
new file mode 100644
--- /dev/null
+++ b/Engine/Plugins/Developer/NamingTokens/Source/NamingTokensEditor/Private/NamingTokensStyle.h
@@ -0,0 +1,20 @@
+// Copyright Epic Games, Inc. All Rights Reserved.
+
+#pragma once
+
+#include "Styling/SlateStyle.h"
+#include "Styling/SlateStyleRegistry.h"
+
+/** Slate style set for Naming Tokens */
+class FNamingTokensStyle : public FSlateStyleSet
+{
+public:
+	static FName StyleName;
+
+	/** Access the singleton instance for this style set */
+	static FNamingTokensStyle& Get();
+
+private:
+	FNamingTokensStyle();
+	~FNamingTokensStyle();
+};
diff --git a/Engine/Plugins/Developer/NamingTokens/Source/NamingTokensEditor/Private/Utils/NamingTokensEditorUtils.cpp b/Engine/Plugins/Developer/NamingTokens/Source/NamingTokensEditor/Private/Utils/NamingTokensEditorUtils.cpp
new file mode 100644
--- /dev/null
+++ b/Engine/Plugins/Developer/NamingTokens/Source/NamingTokensEditor/Private/Utils/NamingTokensEditorUtils.cpp
@@ -0,0 +1,31 @@
+// Copyright Epic Games, Inc. All Rights Reserved.
+
+#include "NamingTokensEditorUtils.h"
+
+#include "Kismet2/BlueprintEditorUtils.h"
+#include "Kismet2/KismetEditorUtilities.h"
+#include "Utils/NamingTokenUtils.h"
+
+FString UE::NamingTokens::Utils::Editor::Private::CreateBaseTokenFunctionName(const FString& InTokenKey)
+{
+	return TEXT("ProcessToken_") + InTokenKey;
+}
+
+FName UE::NamingTokens::Utils::Editor::Private::CreateNewTokenGraph(UBlueprint* InBlueprint, const FString& InTokenKey)
+{
+	check(InBlueprint);
+	
+	const FString BaseName = CreateBaseTokenFunctionName(InTokenKey);
+	const FName FunctionName = FBlueprintEditorUtils::GenerateUniqueGraphName(InBlueprint, BaseName);
+
+	InBlueprint->Modify();
+		
+	UEdGraph* NewGraph = FBlueprintEditorUtils::CreateNewGraph(InBlueprint, FunctionName, UEdGraph::StaticClass(), UEdGraphSchema_K2::StaticClass());
+	
+	UFunction* FunctionSignature = UE::NamingTokens::Utils::GetProcessTokenFunctionSignature();
+		
+	FBlueprintEditorUtils::AddFunctionGraph<UFunction>(InBlueprint, NewGraph, /*bIsUserCreated=*/ true, FunctionSignature);
+	FKismetEditorUtilities::BringKismetToFocusAttentionOnObject(NewGraph);
+
+	return FunctionName;
+}
diff --git a/Engine/Plugins/Developer/NamingTokens/Source/NamingTokensEditor/Private/Utils/NamingTokensEditorUtils.h b/Engine/Plugins/Developer/NamingTokens/Source/NamingTokensEditor/Private/Utils/NamingTokensEditorUtils.h
new file mode 100644
--- /dev/null
+++ b/Engine/Plugins/Developer/NamingTokens/Source/NamingTokensEditor/Private/Utils/NamingTokensEditorUtils.h
@@ -0,0 +1,16 @@
+// Copyright Epic Games, Inc. All Rights Reserved.
+
+#pragma once
+
+#include "CoreMinimal.h"
+
+class UBlueprint;
+
+namespace UE::NamingTokens::Utils::Editor::Private
+{
+	/** Create the default function token name to use, before unique naming adjustments. */
+	FString CreateBaseTokenFunctionName(const FString& InTokenKey);
+	
+	/** Creates a new k2 graph in a blueprint for a given token key. */
+	FName CreateNewTokenGraph(UBlueprint* InBlueprint, const FString& InTokenKey);
+}
diff --git a/Engine/Plugins/Developer/NamingTokens/Source/NamingTokensEditor/Public/EditorUtilityNamingTokens.h b/Engine/Plugins/Developer/NamingTokens/Source/NamingTokensEditor/Public/EditorUtilityNamingTokens.h
new file mode 100644
--- /dev/null
+++ b/Engine/Plugins/Developer/NamingTokens/Source/NamingTokensEditor/Public/EditorUtilityNamingTokens.h
@@ -0,0 +1,27 @@
+// Copyright Epic Games, Inc. All Rights Reserved.
+
+#pragma once
+
+#include "NamingTokens.h"
+
+#include "EditorUtilityNamingTokens.generated.h"
+
+#define UE_API NAMINGTOKENSEDITOR_API
+
+/**
+ * Subclass to define naming tokens for editor use only.
+ */
+UCLASS(MinimalAPI, Blueprintable, Abstract, meta = (ShowWorldContextPin))
+class UEditorUtilityNamingTokens : public UNamingTokens
+{
+	GENERATED_BODY()
+
+public:
+	UE_API UEditorUtilityNamingTokens();
+
+	UE_API virtual void PostLoad() override;
+
+	UE_API virtual bool IsEditorOnly() const override;
+};
+
+#undef UE_API
diff --git a/Engine/Plugins/Developer/NamingTokens/Source/NamingTokensUncookedOnly/NamingTokensUncookedOnly.Build.cs b/Engine/Plugins/Developer/NamingTokens/Source/NamingTokensUncookedOnly/NamingTokensUncookedOnly.Build.cs
--- a/Engine/Plugins/Developer/NamingTokens/Source/NamingTokensUncookedOnly/NamingTokensUncookedOnly.Build.cs
+++ b/Engine/Plugins/Developer/NamingTokens/Source/NamingTokensUncookedOnly/NamingTokensUncookedOnly.Build.cs
@@ -8,26 +8,15 @@ public class NamingTokensUncookedOnly : ModuleRules
     {
         PCHUsage = ModuleRules.PCHUsageMode.UseExplicitOrSharedPCHs;
 
-        PublicDependencyModuleNames.AddRange(
-            new string[]
-            {
-                "Core",
-                "NamingTokens"
-            }
-        );
-
         PrivateDependencyModuleNames.AddRange(
             new string[]
             {
-	            "AssetDefinition",
 	            "BlueprintGraph",
+	            "Blutility",
+	            "Core",
                 "CoreUObject",
                 "Engine",
-                "EngineAssetDefinitions",
-                "InputCore",
-				"Projects",
-                "Slate",
-                "SlateCore",
+                "NamingTokens",
                 "UnrealEd",
             }
         );
diff --git a/Engine/Plugins/Developer/NamingTokens/Source/NamingTokensUncookedOnly/Private/AssetDefinition_NamingTokens.cpp b/Engine/Plugins/Developer/NamingTokens/Source/NamingTokensUncookedOnly/Private/AssetDefinition_NamingTokens.cpp
deleted file mode 100644
diff --git a/Engine/Plugins/Developer/NamingTokens/Source/NamingTokensUncookedOnly/Private/AssetDefinition_NamingTokens.h b/Engine/Plugins/Developer/NamingTokens/Source/NamingTokensUncookedOnly/Private/AssetDefinition_NamingTokens.h
deleted file mode 100644
diff --git a/Engine/Plugins/Developer/NamingTokens/Source/NamingTokensUncookedOnly/Private/Customization/NamingTokensCustomization.cpp b/Engine/Plugins/Developer/NamingTokens/Source/NamingTokensUncookedOnly/Private/Customization/NamingTokensCustomization.cpp
deleted file mode 100644
diff --git a/Engine/Plugins/Developer/NamingTokens/Source/NamingTokensUncookedOnly/Private/Customization/NamingTokensCustomization.h b/Engine/Plugins/Developer/NamingTokens/Source/NamingTokensUncookedOnly/Private/Customization/NamingTokensCustomization.h
deleted file mode 100644
diff --git a/Engine/Plugins/Developer/NamingTokens/Source/NamingTokensUncookedOnly/Private/NamingTokensFactory.cpp b/Engine/Plugins/Developer/NamingTokens/Source/NamingTokensUncookedOnly/Private/NamingTokensFactory.cpp
--- a/Engine/Plugins/Developer/NamingTokens/Source/NamingTokensUncookedOnly/Private/NamingTokensFactory.cpp
+++ b/Engine/Plugins/Developer/NamingTokens/Source/NamingTokensUncookedOnly/Private/NamingTokensFactory.cpp
@@ -4,7 +4,6 @@
 
 #include "NamingTokens.h"
 
-#include "BlueprintEditorSettings.h"
 #include "EdGraphSchema_K2.h"
 #include "Engine/Blueprint.h"
 #include "Engine/BlueprintGeneratedClass.h"
@@ -35,35 +34,6 @@ UObject* UNamingTokensFactory::FactoryCreateNew(UClass* InClass, UObject* InPare
 		UBlueprint::StaticClass(),
 		UBlueprintGeneratedClass::StaticClass()
 	);
-	
-	const UBlueprintEditorSettings* Settings = GetDefault<UBlueprintEditorSettings>();
-	if (Settings && Settings->bSpawnDefaultBlueprintNodes)
-	{
-		// Create default events.
-		int32 NodePositionY = 0;
-
-		UEdGraph* EventGraph = FindObject<UEdGraph>(NewBP, *(UEdGraphSchema_K2::GN_EventGraph.ToString()));
-		check(EventGraph);
-		
-		if (UK2Node_Event* OnPreEvaluateNode = FKismetEditorUtilities::AddDefaultEventNode(NewBP, EventGraph,
-			GET_FUNCTION_NAME_CHECKED(UNamingTokens, OnPreEvaluate), UNamingTokens::StaticClass(), NodePositionY))
-		{
-			// Set the node comment. The comment is displayed because we are automatically placed as ghost nodes.
-			// Once a connection is made the comment will go away, but is still accessible via tooltip.
-			const UFunction* Function = UNamingTokens::StaticClass()->FindFunctionByName(GET_FUNCTION_NAME_CHECKED(UNamingTokens, OnPreEvaluate));
-			check(Function);
-			OnPreEvaluateNode->NodeComment = Function->GetToolTipText().ToString();
-		}
-		if (UK2Node_Event* OnPostEvaluateNode = FKismetEditorUtilities::AddDefaultEventNode(NewBP, EventGraph,
-			GET_FUNCTION_NAME_CHECKED(UNamingTokens, OnPostEvaluate), UNamingTokens::StaticClass(), NodePositionY))
-		{
-			// Set the node comment. The comment is displayed because we are automatically placed as ghost nodes.
-			// Once a connection is made the comment will go away, but is still accessible via tooltip.
-			const UFunction* Function = UNamingTokens::StaticClass()->FindFunctionByName(GET_FUNCTION_NAME_CHECKED(UNamingTokens, OnPostEvaluate));
-			check(Function);
-			OnPostEvaluateNode->NodeComment = Function->GetToolTipText().ToString();
-		}
-	}
 
 	return NewBP;
 }
diff --git a/Engine/Plugins/Developer/NamingTokens/Source/NamingTokensUncookedOnly/Private/NamingTokensStyle.cpp b/Engine/Plugins/Developer/NamingTokens/Source/NamingTokensUncookedOnly/Private/NamingTokensStyle.cpp
deleted file mode 100644
diff --git a/Engine/Plugins/Developer/NamingTokens/Source/NamingTokensUncookedOnly/Private/NamingTokensStyle.h b/Engine/Plugins/Developer/NamingTokens/Source/NamingTokensUncookedOnly/Private/NamingTokensStyle.h
deleted file mode 100644
diff --git a/Engine/Plugins/Developer/NamingTokens/Source/NamingTokensUncookedOnly/Private/NamingTokensUncookedOnlyModule.cpp b/Engine/Plugins/Developer/NamingTokens/Source/NamingTokensUncookedOnly/Private/NamingTokensUncookedOnlyModule.cpp
--- a/Engine/Plugins/Developer/NamingTokens/Source/NamingTokensUncookedOnly/Private/NamingTokensUncookedOnlyModule.cpp
+++ b/Engine/Plugins/Developer/NamingTokens/Source/NamingTokensUncookedOnly/Private/NamingTokensUncookedOnlyModule.cpp
@@ -2,36 +2,57 @@
 
 #include "NamingTokensUncookedOnlyModule.h"
 
-#include "Customization/NamingTokensCustomization.h"
-#include "NamingTokenData.h"
+#include "AssetRegistry/AssetData.h"
+#include "AssetRegistry/AssetRegistryModule.h"
+#include "EditorUtilityBlueprint.h"
 #include "NamingTokens.h"
-#include "NamingTokensStyle.h"
-
-#include "PropertyEditorModule.h"
+#include "Utils/NamingTokensUncookedOnlyUtils.h"
 
 void FNamingTokensUncookedOnlyModule::StartupModule()
 {
-	FPropertyEditorModule& PropertyModule = FModuleManager::GetModuleChecked<FPropertyEditorModule>("PropertyEditor");
-
-	PropertyModule.RegisterCustomClassLayout(
-		UNamingTokens::StaticClass()->GetFName(),
-		FOnGetDetailCustomizationInstance::CreateStatic(&FNamingTokensCustomization::MakeInstance)
-	);
-
-	PropertyModule.RegisterCustomPropertyTypeLayout(FNamingTokenData::StaticStruct()->GetFName(),
-		FOnGetPropertyTypeCustomizationInstance::CreateStatic(&FNamingTokensDataCustomization::MakeInstance));
-
-	FNamingTokensStyle::Get();
+	IAssetRegistry& AssetRegistry = FModuleManager::LoadModuleChecked<FAssetRegistryModule>("AssetRegistry").Get();
+	AssetAddedHandle = AssetRegistry.OnAssetAdded().AddRaw(this, &FNamingTokensUncookedOnlyModule::OnAssetAdded);
 }
 
 void FNamingTokensUncookedOnlyModule::ShutdownModule()
 {
-	if (UObjectInitialized() && !IsEngineExitRequested())
+	if (FModuleManager::Get().IsModuleLoaded("AssetRegistry"))
 	{
-		if (FPropertyEditorModule* PropertyEditorModule = FModuleManager::GetModulePtr<FPropertyEditorModule>("PropertyEditor"))
+		IAssetRegistry& AssetRegistry = FModuleManager::GetModuleChecked<FAssetRegistryModule>("AssetRegistry").Get();
+		AssetRegistry.OnAssetAdded().Remove(AssetAddedHandle);
+	}
+}
+
+void FNamingTokensUncookedOnlyModule::OnAssetAdded(const FAssetData& InAssetData)
+{
+	// Look for any Naming Token blueprints being created and make sure we have the correct starting nodes.
+	// We can't rely on a single factory because we have to account for Editor Utility Blueprint versions too.
+	
+	// Check with fast methods to determine if we are dealing with our asset or not. The asset should
+	// already be loaded into memory.
+	if ((InAssetData.IsValid() && !InAssetData.IsRedirector() && InAssetData.IsAssetLoaded())
+		&& (InAssetData.AssetClassPath == UBlueprint::StaticClass()->GetClassPathName()
+			|| InAssetData.AssetClassPath == UEditorUtilityBlueprint::StaticClass()->GetClassPathName()))
+	{
+		FString NativeParentPathStr;
+		if (!InAssetData.GetTagValue(FBlueprintTags::NativeParentClassPath, NativeParentPathStr))
+		{
+			return;
+		}
+		const FSoftClassPath NativeParentPath(NativeParentPathStr);
+		const UClass* NativeParentClass = NativeParentPath.ResolveClass();
+		if (!NativeParentClass)
+		{
+			NativeParentClass = NativeParentPath.TryLoadClass<UObject>();
+		}
+			
+		if (NativeParentClass && NativeParentClass->IsChildOf(UNamingTokens::StaticClass()))
 		{
-			PropertyEditorModule->UnregisterCustomClassLayout(UNamingTokens::StaticClass()->GetFName());
-			PropertyEditorModule->UnregisterCustomPropertyTypeLayout(FNamingTokenData::StaticStruct()->GetFName());
+			UBlueprint* Blueprint = CastChecked<UBlueprint>(InAssetData.GetAsset());
+			if (Blueprint->bIsNewlyCreated)
+			{
+				UE::NamingTokens::Utils::UncookedOnly::Private::SetupInitialBlueprint(Blueprint);
+			}
 		}
 	}
 }
diff --git a/Engine/Plugins/Developer/NamingTokens/Source/NamingTokensUncookedOnly/Private/NamingTokensUncookedOnlyModule.h b/Engine/Plugins/Developer/NamingTokens/Source/NamingTokensUncookedOnly/Private/NamingTokensUncookedOnlyModule.h
--- a/Engine/Plugins/Developer/NamingTokens/Source/NamingTokensUncookedOnly/Private/NamingTokensUncookedOnlyModule.h
+++ b/Engine/Plugins/Developer/NamingTokens/Source/NamingTokensUncookedOnly/Private/NamingTokensUncookedOnlyModule.h
@@ -5,9 +5,17 @@
 #include "CoreMinimal.h"
 #include "Modules/ModuleManager.h"
 
+struct FAssetData;
+
 class FNamingTokensUncookedOnlyModule : public IModuleInterface
 {
 public:
     virtual void StartupModule() override;
     virtual void ShutdownModule() override;
+
+private:
+	void OnAssetAdded(const FAssetData& InAssetData);
+	
+private:
+	FDelegateHandle AssetAddedHandle;
 };
diff --git a/Engine/Plugins/Developer/NamingTokens/Source/NamingTokensUncookedOnly/Private/Utils/NamingTokensEditorUtils.cpp b/Engine/Plugins/Developer/NamingTokens/Source/NamingTokensUncookedOnly/Private/Utils/NamingTokensEditorUtils.cpp
deleted file mode 100644
diff --git a/Engine/Plugins/Developer/NamingTokens/Source/NamingTokensUncookedOnly/Private/Utils/NamingTokensEditorUtils.h b/Engine/Plugins/Developer/NamingTokens/Source/NamingTokensUncookedOnly/Private/Utils/NamingTokensEditorUtils.h
deleted file mode 100644
diff --git a/Engine/Plugins/Developer/NamingTokens/Source/NamingTokensUncookedOnly/Private/Utils/NamingTokensUncookedOnlyUtils.cpp b/Engine/Plugins/Developer/NamingTokens/Source/NamingTokensUncookedOnly/Private/Utils/NamingTokensUncookedOnlyUtils.cpp
new file mode 100644
--- /dev/null
+++ b/Engine/Plugins/Developer/NamingTokens/Source/NamingTokensUncookedOnly/Private/Utils/NamingTokensUncookedOnlyUtils.cpp
@@ -0,0 +1,41 @@
+// Copyright Epic Games, Inc. All Rights Reserved.
+
+#include "Utils/NamingTokensUncookedOnlyUtils.h"
+
+#include "BlueprintEditorSettings.h"
+#include "EdGraphSchema_K2.h"
+#include "K2Node_Event.h"
+#include "Kismet2/KismetEditorUtilities.h"
+#include "NamingTokens.h"
+
+void UE::NamingTokens::Utils::UncookedOnly::Private::SetupInitialBlueprint(UBlueprint* InBlueprint)
+{
+	const UBlueprintEditorSettings* Settings = GetDefault<UBlueprintEditorSettings>();
+	if (Settings && Settings->bSpawnDefaultBlueprintNodes)
+	{
+		// Create default events.
+		int32 NodePositionY = 0;
+
+		UEdGraph* EventGraph = FindObject<UEdGraph>(InBlueprint, *(UEdGraphSchema_K2::GN_EventGraph.ToString()));
+		check(EventGraph);
+		
+		if (UK2Node_Event* OnPreEvaluateNode = FKismetEditorUtilities::AddDefaultEventNode(InBlueprint, EventGraph,
+			UNamingTokens::GetOnPreEvaluateFunctionName(), UNamingTokens::StaticClass(), NodePositionY))
+		{
+			// Set the node comment. The comment is displayed because we are automatically placed as ghost nodes.
+			// Once a connection is made the comment will go away, but is still accessible via tooltip.
+			const UFunction* Function = UNamingTokens::StaticClass()->FindFunctionByName(UNamingTokens::GetOnPreEvaluateFunctionName());
+			check(Function);
+			OnPreEvaluateNode->NodeComment = Function->GetToolTipText().ToString();
+		}
+		if (UK2Node_Event* OnPostEvaluateNode = FKismetEditorUtilities::AddDefaultEventNode(InBlueprint, EventGraph,
+			UNamingTokens::GetOnPostEvaluateFunctionName(), UNamingTokens::StaticClass(), NodePositionY))
+		{
+			// Set the node comment. The comment is displayed because we are automatically placed as ghost nodes.
+			// Once a connection is made the comment will go away, but is still accessible via tooltip.
+			const UFunction* Function = UNamingTokens::StaticClass()->FindFunctionByName(UNamingTokens::GetOnPostEvaluateFunctionName());
+			check(Function);
+			OnPostEvaluateNode->NodeComment = Function->GetToolTipText().ToString();
+		}
+	}
+}
diff --git a/Engine/Plugins/Developer/NamingTokens/Source/NamingTokensUncookedOnly/Private/Utils/NamingTokensUncookedOnlyUtils.h b/Engine/Plugins/Developer/NamingTokens/Source/NamingTokensUncookedOnly/Private/Utils/NamingTokensUncookedOnlyUtils.h
new file mode 100644
--- /dev/null
+++ b/Engine/Plugins/Developer/NamingTokens/Source/NamingTokensUncookedOnly/Private/Utils/NamingTokensUncookedOnlyUtils.h
@@ -0,0 +1,11 @@
+// Copyright Epic Games, Inc. All Rights Reserved.
+
+#pragma once
+
+class UBlueprint;
+
+namespace UE::NamingTokens::Utils::UncookedOnly::Private
+{
+	/** Configure a new blueprint for Naming Token use. */
+	void SetupInitialBlueprint(UBlueprint* InBlueprint);
+}
