NamingTokens: Typing in whitepace in a token or namespace will remove it during evaluation.
Add bNormalizeInput option to token filter to allow this to be disabled.

#Synced-CL 48820685

diff --git a/Engine/Plugins/Developer/NamingTokens/Source/NamingTokens/Private/NamingTokens.cpp b/Engine/Plugins/Developer/NamingTokens/Source/NamingTokens/Private/NamingTokens.cpp
--- a/Engine/Plugins/Developer/NamingTokens/Source/NamingTokens/Private/NamingTokens.cpp
+++ b/Engine/Plugins/Developer/NamingTokens/Source/NamingTokens/Private/NamingTokens.cpp
@@ -5,6 +5,7 @@
 #include "NamingTokensEvaluationData.h"
 #include "NamingTokensLog.h"
 #include "Utils/NamingTokenUtils.h"
+#include "Utils/NamingTokenUtilsPrivate.h"
 
 #include "Engine/Engine.h"
 #include "Internationalization/LocKeyFuncs.h"
@@ -126,10 +127,11 @@ FNamingTokenResultData UNamingTokens::EvaluateTokenText(const FText& InTokenText
 	FNamingTokenResultData Result;
 	Result.OriginalText = InTokenText;
 
-	FText FormattedText = InTokenText;
+	const FString NormalizedString = InEvaluationData.bNormalizeInput ? UE::NamingTokens::Utils::Private::NormalizeTokenString(InTokenText.ToString()) : InTokenText.ToString();
+	FText FormattedText = InEvaluationData.bNormalizeInput ? FText::FromString(NormalizedString) : InTokenText;
 
 	// Retrieve all unprocessed tokens within this text.
-	const TArray<FString> UnprocessedTokens = UE::NamingTokens::Utils::GetTokenKeysFromString(FormattedText.ToString());
+	const TArray<FString> UnprocessedTokens = UE::NamingTokens::Utils::GetTokenKeysFromString(NormalizedString, false);
 
 	TSet<FNamingTokenData> DefaultTokensSet(GetDefaultTokens());
 	TSet<FNamingTokenData> CustomTokenSet(CustomTokens);
diff --git a/Engine/Plugins/Developer/NamingTokens/Source/NamingTokens/Private/NamingTokensEngineSubsystem.cpp b/Engine/Plugins/Developer/NamingTokens/Source/NamingTokens/Private/NamingTokensEngineSubsystem.cpp
--- a/Engine/Plugins/Developer/NamingTokens/Source/NamingTokens/Private/NamingTokensEngineSubsystem.cpp
+++ b/Engine/Plugins/Developer/NamingTokens/Source/NamingTokens/Private/NamingTokensEngineSubsystem.cpp
@@ -66,14 +66,15 @@ UNamingTokensEngineSubsystem::UNamingTokensEngineSubsystem()
 UNamingTokens* UNamingTokensEngineSubsystem::GetNamingTokens(const FString& InNamespace) const
 {
 	// Check cache.
-	UNamingTokens* FoundTokens = GetNamingTokenFromCache(InNamespace, false);
+	const FString SanitizedNamespace = UE::NamingTokens::Utils::Private::SanitizeTokenOrNamespace(InNamespace);
+	UNamingTokens* FoundTokens = GetNamingTokenFromCache(SanitizedNamespace, false);
 	
 	if (FoundTokens)
 	{
 		return FoundTokens;
 	}
 
-	FoundTokens = GetNamingTokensNative(InNamespace);
+	FoundTokens = GetNamingTokensNative(SanitizedNamespace);
 
 	const FAssetRegistryModule& AssetRegistryModule = FModuleManager::LoadModuleChecked<FAssetRegistryModule>("AssetRegistry");
 	const IAssetRegistry& AssetRegistry = AssetRegistryModule.Get();
@@ -81,15 +82,15 @@ UNamingTokens* UNamingTokensEngineSubsystem::GetNamingTokens(const FString& InNa
 	// Check blueprint classes.
 	// NOTE: We may want to use OnAssetAdded and cache the asset or load the class then.
 	UE::NamingTokens::Utils::Private::ForEachNamingTokensBlueprint(
-		[&InNamespace, &AssetRegistry, &FoundTokens, this](const FAssetData& AssetData, const FTopLevelAssetPath& ClassObjectPath)
+		[&SanitizedNamespace, &AssetRegistry, &FoundTokens, this](const FAssetData& AssetData, const FTopLevelAssetPath& ClassObjectPath)
 	{
 		// Now check for the namespace match.
 		FString FoundNamespace;
 		if (AssetData.GetTagValue(UNamingTokens::GetNamespacePropertyName(), FoundNamespace)
-			&& FoundNamespace == InNamespace)
+			&& FoundNamespace == SanitizedNamespace)
 		{
 			// Allow BP children to overwrite native classes.
-			auto IsAssetChildOfNativeToken = [&AssetRegistry, &InNamespace](const FTopLevelAssetPath& InAssetPath, const UNamingTokens* InParentTokens) -> bool
+			auto IsAssetChildOfNativeToken = [&AssetRegistry, &SanitizedNamespace](const FTopLevelAssetPath& InAssetPath, const UNamingTokens* InParentTokens) -> bool
 			{
 				check(InParentTokens);
 				if (InParentTokens->GetClass()->IsNative())
@@ -102,7 +103,7 @@ UNamingTokens* UNamingTokensEngineSubsystem::GetNamingTokens(const FString& InNa
 					if (NativeDerivedClassNames.Contains(InAssetPath))
 					{
 						UE_LOG(LogNamingTokens, Log, TEXT("Using namespace '%s' of BP child '%s' instead of native parent '%s'."),
-							*InNamespace, *InAssetPath.ToString(), *InParentTokens->GetClass()->GetName());
+							*SanitizedNamespace, *InAssetPath.ToString(), *InParentTokens->GetClass()->GetName());
 						return true;
 					}
 				}
@@ -112,12 +113,12 @@ UNamingTokens* UNamingTokensEngineSubsystem::GetNamingTokens(const FString& InNa
 			if (!FoundTokens || IsAssetChildOfNativeToken(ClassObjectPath, FoundTokens))
 			{
 				const TSoftClassPtr<UNamingTokens> SoftClassPath(ClassObjectPath.ToString());
-				FoundTokens = LoadNamingToken(SoftClassPath, InNamespace);
+				FoundTokens = LoadNamingToken(SoftClassPath, SanitizedNamespace);
 			}
 			else
 			{
 				UE_LOG(LogNamingTokens, Warning, TEXT("Namespace '%s' exists more than once in class '%s' and BP asset '%s'."),
-					*InNamespace, *FoundTokens->GetClass()->GetName(), *AssetData.PackageName.ToString());
+					*SanitizedNamespace, *FoundTokens->GetClass()->GetName(), *AssetData.PackageName.ToString());
 			}
 		}
 	});
@@ -127,8 +128,9 @@ UNamingTokens* UNamingTokensEngineSubsystem::GetNamingTokens(const FString& InNa
 
 UNamingTokens* UNamingTokensEngineSubsystem::GetNamingTokensNative(const FString& InNamespace) const
 {
+	const FString SanitizedNamespace = UE::NamingTokens::Utils::Private::SanitizeTokenOrNamespace(InNamespace);
 	// Check cache.
-	UNamingTokens* FoundTokens = GetNamingTokenFromCache(InNamespace, true);
+	UNamingTokens* FoundTokens = GetNamingTokenFromCache(SanitizedNamespace, true);
 	if (FoundTokens)
 	{
 		return FoundTokens;
@@ -148,16 +150,16 @@ UNamingTokens* UNamingTokensEngineSubsystem::GetNamingTokensNative(const FString
 			&& !DerivedClass->HasAnyClassFlags(InvalidClassFlags)
 			&& !DerivedClass->GetName().StartsWith(TEXT("SKEL_"))
 			&& !DerivedClass->GetName().StartsWith(TEXT("REINST_"))
-			&& DerivedClass->GetDefaultObject<UNamingTokens>()->GetNamespace() == InNamespace)
+			&& DerivedClass->GetDefaultObject<UNamingTokens>()->GetNamespace() == SanitizedNamespace)
 		{
 			if (!FoundTokens)
 			{
-				FoundTokens = LoadNamingToken(DerivedClass, InNamespace);
+				FoundTokens = LoadNamingToken(DerivedClass, SanitizedNamespace);
 			}
 			else
 			{
 				UE_LOG(LogNamingTokens, Warning, TEXT("Namespace '%s' exists more than once in native class '%s' and '%s'."),
-						*InNamespace, *FoundTokens->GetClass()->GetName(), *DerivedClass->GetName());
+						*SanitizedNamespace, *FoundTokens->GetClass()->GetName(), *DerivedClass->GetName());
 			}
 		}
 	}
@@ -204,6 +206,7 @@ FNamingTokenResultData UNamingTokensEngineSubsystem::EvaluateTokenText(const FTe
 	EvaluationData.Initialize();
 	EvaluationData.Contexts = InContexts;
 	EvaluationData.bForceCaseSensitive = InFilter.bForceCaseSensitive;
+	EvaluationData.bNormalizeInput = InFilter.bNormalizeInput;
 	for (const FString& Namespace : Namespaces)
 	{
 		DECLARE_DELEGATE_RetVal_OneParam(UNamingTokens*, FGetNamingTokensFunction, const FString&);
@@ -344,7 +347,7 @@ TSet<FString> UNamingTokensEngineSubsystem::GetNamingTokenNamespacesFromString(c
 {
 	TSet<FString> Namespaces;
 
-	const FRegexPattern Pattern(FString::Printf(TEXT(R"(\{\s*([a-zA-Z0-9_]+)%s[a-zA-Z0-9_]+\s*\})"), *UE::NamingTokens::Utils::GetNamespaceDelimiter()));
+	const FRegexPattern Pattern(FString::Printf(TEXT(R"(\{\s*([a-zA-Z0-9_ ]+)\s*%s\s*[a-zA-Z0-9_ ]+\s*\})"), *UE::NamingTokens::Utils::GetNamespaceDelimiter()));
 	FRegexMatcher Matcher(Pattern, InTokenString);
 	
 	while (Matcher.FindNext())
@@ -428,14 +431,15 @@ UNamingTokens* UNamingTokensEngineSubsystem::LoadNamingToken(const TSoftClassPtr
 UNamingTokens* UNamingTokensEngineSubsystem::GetNamingTokenFromCache(const FString& InNamespace, bool bNativeOnly) const
 {
 	FScopeLock Lock(&CachedNamingTokensMutex);
+	const FString SanitizedNamespace = UE::NamingTokens::Utils::Private::SanitizeTokenOrNamespace(InNamespace);
 	if (IsCacheEnabled())
 	{
-		if (const TObjectPtr<UNamingTokens>* Tokens = CachedNamingTokens.Find(InNamespace))
+		if (const TObjectPtr<UNamingTokens>* Tokens = CachedNamingTokens.Find(SanitizedNamespace))
 		{
 			if ((*Tokens) == nullptr)
 			{
 				// NamingTokens may have been deleted, with a new asset using the same namespace then added afterward.
-				CachedNamingTokens.Remove(InNamespace);
+				CachedNamingTokens.Remove(SanitizedNamespace);
 				return nullptr;
 			}
 			
diff --git a/Engine/Plugins/Developer/NamingTokens/Source/NamingTokens/Private/Utils/NamingTokenUtils.cpp b/Engine/Plugins/Developer/NamingTokens/Source/NamingTokens/Private/Utils/NamingTokenUtils.cpp
--- a/Engine/Plugins/Developer/NamingTokens/Source/NamingTokens/Private/Utils/NamingTokenUtils.cpp
+++ b/Engine/Plugins/Developer/NamingTokens/Source/NamingTokens/Private/Utils/NamingTokenUtils.cpp
@@ -2,19 +2,12 @@
 
 #include "Utils/NamingTokenUtils.h"
 
+#include "Utils/NamingTokenUtilsPrivate.h"
 #include "NamingTokens.h"
 
 #include "Internationalization/LocKeyFuncs.h"
 #include "Internationalization/Regex.h"
 
-namespace UE::NamingTokens::Utils::Private
-{
-	FString GetTokenPatternString()
-	{
-		return FString::Printf(TEXT(R"(\{\s*((?:[a-zA-Z0-9_]+%s)*[a-zA-Z0-9_]+)\s*\})"), *GetNamespaceDelimiter());
-	}
-}
-
 UFunction* UE::NamingTokens::Utils::GetProcessTokenFunctionSignature()
 {
 	UFunction* FunctionSignature = UNamingTokens::StaticClass()->FindFunctionByName(UNamingTokens::GetProcessTokenTemplateFunctionName());
@@ -45,17 +38,19 @@ FString UE::NamingTokens::Utils::GetNamespaceDelimiter()
 	return Delimiter;
 }
 
-TArray<FString> UE::NamingTokens::Utils::GetTokenKeysFromString(const FString& InTokenString)
+TArray<FString> UE::NamingTokens::Utils::GetTokenKeysFromString(const FString& InTokenString, bool bNormalizeString)
 {
+	const FString NormalizedString = bNormalizeString ? Private::NormalizeTokenString(InTokenString) : InTokenString;
 	static const FString PatternString = Private::GetTokenPatternString();
 	static const FRegexPattern Pattern(PatternString);
-
+	
 	TSet<FString, FLocKeySetFuncs> Tokens;
-	FRegexMatcher Matcher(Pattern, InTokenString);
+	FRegexMatcher Matcher(Pattern, NormalizedString);
 
 	while (Matcher.FindNext())
 	{
 		FString Token = Matcher.GetCaptureGroup(1);
+		Token.RemoveSpacesInline();
 		Tokens.Add(Token);
 	}
 
diff --git a/Engine/Plugins/Developer/NamingTokens/Source/NamingTokens/Private/Utils/NamingTokenUtilsPrivate.cpp b/Engine/Plugins/Developer/NamingTokens/Source/NamingTokens/Private/Utils/NamingTokenUtilsPrivate.cpp
--- a/Engine/Plugins/Developer/NamingTokens/Source/NamingTokens/Private/Utils/NamingTokenUtilsPrivate.cpp
+++ b/Engine/Plugins/Developer/NamingTokens/Source/NamingTokens/Private/Utils/NamingTokenUtilsPrivate.cpp
@@ -6,8 +6,10 @@
 #include "AssetRegistry/AssetDataTagMap.h"
 #include "AssetRegistry/AssetRegistryModule.h"
 #include "Engine/Blueprint.h"
+#include "Internationalization/Regex.h"
 #include "Misc/PackageName.h"
 #include "NamingTokens.h"
+#include "Utils/NamingTokenUtils.h"
 
 void UE::NamingTokens::Utils::Private::ForEachNamingTokensBlueprint(
 	const TFunction<void(const FAssetData& AssetData, const FTopLevelAssetPath& ClassObjectPath)>& Function)
@@ -42,3 +44,46 @@ void UE::NamingTokens::Utils::Private::ForEachNamingTokensBlueprint(
 		}
 	}
 }
+
+FString UE::NamingTokens::Utils::Private::NormalizeTokenString(const FString& InTokenString)
+{
+	const FRegexPattern Pattern(TEXT(R"(\{([^}]*)\})"));
+	FRegexMatcher Matcher(Pattern, InTokenString);
+
+	FString Output;
+	int32 LastIndex = 0;
+
+	while (Matcher.FindNext())
+	{
+		const int32 Start = Matcher.GetMatchBeginning();
+		const int32 End = Matcher.GetMatchEnding();
+
+		Output += InTokenString.Mid(LastIndex, Start - LastIndex);
+
+		FString Inner = Matcher.GetCaptureGroup(1);
+		Inner.ReplaceInline(TEXT(" "), TEXT(""));
+		Inner.ReplaceInline(TEXT("\t"), TEXT(""));
+		Inner.ReplaceInline(TEXT("\n"), TEXT(""));
+		Inner.ReplaceInline(TEXT("\r"), TEXT(""));
+		
+		Output += FString::Printf(TEXT("{%s}"), *Inner);
+
+		LastIndex = End;
+	}
+
+	Output += InTokenString.Mid(LastIndex);
+	return Output;
+}
+
+FString UE::NamingTokens::Utils::Private::SanitizeTokenOrNamespace(const FString& InString)
+{
+	FString SanitizedString = InString;
+	SanitizedString.RemoveSpacesInline();
+	
+	return SanitizedString;
+}
+
+FString UE::NamingTokens::Utils::Private::GetTokenPatternString()
+{
+	return FString::Printf(TEXT(R"(\{\s*((?:[a-zA-Z0-9_]+%s)*[a-zA-Z0-9_]+)\s*\})"), *GetNamespaceDelimiter());
+}
diff --git a/Engine/Plugins/Developer/NamingTokens/Source/NamingTokens/Private/Utils/NamingTokenUtilsPrivate.h b/Engine/Plugins/Developer/NamingTokens/Source/NamingTokens/Private/Utils/NamingTokenUtilsPrivate.h
--- a/Engine/Plugins/Developer/NamingTokens/Source/NamingTokens/Private/Utils/NamingTokenUtilsPrivate.h
+++ b/Engine/Plugins/Developer/NamingTokens/Source/NamingTokens/Private/Utils/NamingTokenUtilsPrivate.h
@@ -11,4 +11,13 @@ namespace UE::NamingTokens::Utils::Private
 {
 	/** Iterate each Naming Tokens blueprint asset, executing a function for each qualifying asset. */
 	void ForEachNamingTokensBlueprint(const TFunction<void(const FAssetData& AssetData, const FTopLevelAssetPath& ClassObjectPath)>& Function);
+	
+	/** Preprocess a string removing any spaces within braces. */
+	FString NormalizeTokenString(const FString& InTokenString);
+	
+	/** Remove spaces from within tokens or namespaces. */
+	FString SanitizeTokenOrNamespace(const FString& InString);
+	
+	/** Retrieve the token regex. */
+	FString GetTokenPatternString();
 }
diff --git a/Engine/Plugins/Developer/NamingTokens/Source/NamingTokens/Public/NamingTokensEngineSubsystem.h b/Engine/Plugins/Developer/NamingTokens/Source/NamingTokens/Public/NamingTokensEngineSubsystem.h
--- a/Engine/Plugins/Developer/NamingTokens/Source/NamingTokens/Public/NamingTokensEngineSubsystem.h
+++ b/Engine/Plugins/Developer/NamingTokens/Source/NamingTokens/Public/NamingTokensEngineSubsystem.h
@@ -35,8 +35,15 @@ struct FNamingTokenFilterArgs
 	bool bForceCaseSensitive = false;
 
 	/** When false, we additionally look for blueprint naming tokens. */
-	UPROPERTY(EditAnywhere, BlueprintReadWrite, Category = NamingTokens)
+	UPROPERTY(EditAnywhere, BlueprintReadWrite, Category=NamingTokens)
 	bool bNativeOnly = false;
+	
+	/**
+	 * If input data should be normalized first, such as removing white spaces in tokens.
+	 * A better user experience typing in tokens, but may impact localization data.
+	 */
+	UPROPERTY(EditAnywhere, BlueprintReadWrite, Category=NamingTokens)
+	bool bNormalizeInput = true;
 };
 
 /**
diff --git a/Engine/Plugins/Developer/NamingTokens/Source/NamingTokens/Public/NamingTokensEvaluationData.h b/Engine/Plugins/Developer/NamingTokens/Source/NamingTokens/Public/NamingTokensEvaluationData.h
--- a/Engine/Plugins/Developer/NamingTokens/Source/NamingTokens/Public/NamingTokensEvaluationData.h
+++ b/Engine/Plugins/Developer/NamingTokens/Source/NamingTokens/Public/NamingTokensEvaluationData.h
@@ -31,4 +31,7 @@ public:
 
 	/** If case sensitivity should be enforced. */
 	bool bForceCaseSensitive = false;
+	
+	/** If input data should be normalized first. */
+	bool bNormalizeInput = true;
 };
\ No newline at end of file
diff --git a/Engine/Plugins/Developer/NamingTokens/Source/NamingTokens/Public/Utils/NamingTokenUtils.h b/Engine/Plugins/Developer/NamingTokens/Source/NamingTokens/Public/Utils/NamingTokenUtils.h
--- a/Engine/Plugins/Developer/NamingTokens/Source/NamingTokens/Public/Utils/NamingTokenUtils.h
+++ b/Engine/Plugins/Developer/NamingTokens/Source/NamingTokens/Public/Utils/NamingTokenUtils.h
@@ -20,8 +20,12 @@ namespace UE::NamingTokens::Utils
 	/** Retrieve the delimiter for the namespace. */
 	NAMINGTOKENS_API FString GetNamespaceDelimiter();
 	
-	/** Extract all token keys from a string. */
-	NAMINGTOKENS_API TArray<FString> GetTokenKeysFromString(const FString& InTokenString);
+	/** 
+	 * Extract all token keys from a string.
+	 * @param InTokenString The string to process.
+	 * @param bNormalizeString If the string should be normalized first, stripping any whitespace within braces.
+	 */
+	NAMINGTOKENS_API TArray<FString> GetTokenKeysFromString(const FString& InTokenString, bool bNormalizeString=true);
 
 	/** Checks if a token is present in a string. */
 	NAMINGTOKENS_API bool IsTokenInString(const FString& InTokenKey, const FString& InTokenString);
diff --git a/Engine/Plugins/Developer/NamingTokens/Source/NamingTokensUI/Private/SNamingTokensDataTreeView.cpp b/Engine/Plugins/Developer/NamingTokens/Source/NamingTokensUI/Private/SNamingTokensDataTreeView.cpp
--- a/Engine/Plugins/Developer/NamingTokens/Source/NamingTokensUI/Private/SNamingTokensDataTreeView.cpp
+++ b/Engine/Plugins/Developer/NamingTokens/Source/NamingTokensUI/Private/SNamingTokensDataTreeView.cpp
@@ -266,19 +266,20 @@ void SNamingTokenDataTreeViewWidget::PopulateTreeItems(const FNamingTokenFilterA
 void SNamingTokenDataTreeViewWidget::FilterTreeItems(const FString& InFilterText)
 {
 	bool bSelectingNamespace = false;
-	if (LastFilterText == InFilterText && FilteredRootTreeItems.Num() > 0)
+	FString SanitizedFilterText = InFilterText;
+	SanitizedFilterText.RemoveSpacesInline();
+	if (LastFilterText == SanitizedFilterText && FilteredRootTreeItems.Num() > 0)
 	{
 		if (const FNamingTokenDataTreeItemPtr CurrentlySelectedItem = GetSelectedItem())
 		{
 			// This is used at the end to skip selecting the namespace again if that's all we're showing.
 			// We may need to move this logic elsewhere in case filtering is called when selecting for some reason.
-			bSelectingNamespace = InFilterText == CurrentlySelectedItem->Namespace && CurrentlySelectedItem->IsNamespace();
+			bSelectingNamespace = SanitizedFilterText == CurrentlySelectedItem->Namespace && CurrentlySelectedItem->IsNamespace();
 		}
 	}
-	LastFilterText = InFilterText;
-		
-	FString CompleteOrPartialNamespace = UE::NamingTokens::Utils::GetNamespaceFromTokenKey(InFilterText);
-	FString PartialToken = UE::NamingTokens::Utils::RemoveNamespaceFromTokenKey(InFilterText);
+	LastFilterText = SanitizedFilterText;
+	FString CompleteOrPartialNamespace = UE::NamingTokens::Utils::GetNamespaceFromTokenKey(SanitizedFilterText);
+	FString PartialToken = UE::NamingTokens::Utils::RemoveNamespaceFromTokenKey(SanitizedFilterText);
 
 	TokenFilterText = PartialToken;
 	NamespaceFilterText = CompleteOrPartialNamespace;
@@ -288,8 +289,8 @@ void SNamingTokenDataTreeViewWidget::FilterTreeItems(const FString& InFilterText
 	{
 		// Can't discern either part, assume searching for namespaces.
 		
-		CompleteOrPartialNamespace = InFilterText;
-		PartialToken = InFilterText;
+		CompleteOrPartialNamespace = SanitizedFilterText;
+		PartialToken = SanitizedFilterText;
 	}
 
 	// We have a complete namespace we require.
