Avoid work OnAssetAdded when we are a cooked asset

#Synced-CL 50966172
#Synced-CL 50966172

#Synced-CL 50966172

diff --git a/Engine/Plugins/Developer/NamingTokens/Source/NamingTokensUncookedOnly/Private/NamingTokensUncookedOnlyModule.cpp b/Engine/Plugins/Developer/NamingTokens/Source/NamingTokensUncookedOnly/Private/NamingTokensUncookedOnlyModule.cpp
--- a/Engine/Plugins/Developer/NamingTokens/Source/NamingTokensUncookedOnly/Private/NamingTokensUncookedOnlyModule.cpp
+++ b/Engine/Plugins/Developer/NamingTokens/Source/NamingTokensUncookedOnly/Private/NamingTokensUncookedOnlyModule.cpp
@@ -30,7 +30,7 @@ void FNamingTokensUncookedOnlyModule::OnAssetAdded(const FAssetData& InAssetData
 	
 	// Check with fast methods to determine if we are dealing with our asset or not. The asset should
 	// already be loaded into memory.
-	if ((InAssetData.IsValid() && !InAssetData.IsRedirector() && InAssetData.IsAssetLoaded())
+	if ((InAssetData.IsValid() && !InAssetData.IsRedirector() && !InAssetData.HasAnyPackageFlags(PKG_Cooked) && InAssetData.IsAssetLoaded())
 		&& (InAssetData.AssetClassPath == UBlueprint::StaticClass()->GetClassPathName()
 			|| InAssetData.AssetClassPath == UEditorUtilityBlueprint::StaticClass()->GetClassPathName()))
 	{
